ElectraCodecs: non unity build fixes

#rnx

diff --git a/Engine/Plugins/Media/ElectraCodecs/Source/ElectraDecoders/Private/Android/aac/AudioDecoderAAC_JavaWrapper_Android.h b/Engine/Plugins/Media/ElectraCodecs/Source/ElectraDecoders/Private/Android/aac/AudioDecoderAAC_JavaWrapper_Android.h
--- a/Engine/Plugins/Media/ElectraCodecs/Source/ElectraDecoders/Private/Android/aac/AudioDecoderAAC_JavaWrapper_Android.h
+++ b/Engine/Plugins/Media/ElectraCodecs/Source/ElectraDecoders/Private/Android/aac/AudioDecoderAAC_JavaWrapper_Android.h
@@ -4,6 +4,12 @@
 
 #include "CoreMinimal.h"
 
+#include "HAL/Platform.h"
+
+#if (defined(USE_ANDROID_JNI_WITHOUT_GAMEACTIVITY) && USE_ANDROID_JNI_WITHOUT_GAMEACTIVITY != 0) || USE_ANDROID_JNI
+#include "Android/AndroidPlatform.h"
+#include "Android/AndroidJava.h"
+
 #include "Utils/MPEG/ElectraUtilsMPEGAudio.h"
 
 
@@ -169,3 +175,5 @@ public:
 	virtual int32 GetOutputBufferAndRelease(void*& OutBufferDataPtr, int32 OutBufferDataSize, const FOutputBufferInfo& InOutBufferInfo) = 0;
 
 };
+
+#endif
diff --git a/Engine/Plugins/Media/ElectraCodecs/Source/ElectraDecoders/Private/Android/h264/VideoDecoderH264_JavaWrapper_Android.h b/Engine/Plugins/Media/ElectraCodecs/Source/ElectraDecoders/Private/Android/h264/VideoDecoderH264_JavaWrapper_Android.h
--- a/Engine/Plugins/Media/ElectraCodecs/Source/ElectraDecoders/Private/Android/h264/VideoDecoderH264_JavaWrapper_Android.h
+++ b/Engine/Plugins/Media/ElectraCodecs/Source/ElectraDecoders/Private/Android/h264/VideoDecoderH264_JavaWrapper_Android.h
@@ -4,6 +4,12 @@
 
 #include "CoreMinimal.h"
 
+#include "HAL/Platform.h"
+
+#if (defined(USE_ANDROID_JNI_WITHOUT_GAMEACTIVITY) && USE_ANDROID_JNI_WITHOUT_GAMEACTIVITY != 0) || USE_ANDROID_JNI
+#include "Android/AndroidPlatform.h"
+#include "Android/AndroidJava.h"
+
 #include "Utils/MPEG/ElectraUtilsMPEGVideo.h"
 
 class IElectraH264VideoDecoderAndroidJava
@@ -254,3 +260,4 @@ public:
 	virtual int32 ReleaseOutputBuffer(int32 BufferIndex, int32 ValidCount, bool bRender, int64 releaseAt) = 0;
 };
 
+#endif
diff --git a/Engine/Plugins/Media/ElectraCodecs/Source/ElectraDecoders/Private/Android/h265/VideoDecoderH265_JavaWrapper_Android.h b/Engine/Plugins/Media/ElectraCodecs/Source/ElectraDecoders/Private/Android/h265/VideoDecoderH265_JavaWrapper_Android.h
--- a/Engine/Plugins/Media/ElectraCodecs/Source/ElectraDecoders/Private/Android/h265/VideoDecoderH265_JavaWrapper_Android.h
+++ b/Engine/Plugins/Media/ElectraCodecs/Source/ElectraDecoders/Private/Android/h265/VideoDecoderH265_JavaWrapper_Android.h
@@ -4,6 +4,12 @@
 
 #include "CoreMinimal.h"
 
+#include "HAL/Platform.h"
+
+#if (defined(USE_ANDROID_JNI_WITHOUT_GAMEACTIVITY) && USE_ANDROID_JNI_WITHOUT_GAMEACTIVITY != 0) || USE_ANDROID_JNI
+#include "Android/AndroidPlatform.h"
+#include "Android/AndroidJava.h"
+
 #include "Utils/MPEG/ElectraUtilsMPEGVideo.h"
 
 class IElectraH265VideoDecoderAndroidJava
@@ -255,3 +261,4 @@ public:
 	virtual int32 ReleaseOutputBuffer(int32 BufferIndex, int32 ValidCount, bool bRender, int64 releaseAt) = 0;
 };
 
+#endif
