ElectraCodecs: adjusted code for a potentially iOS16 specific audio decoder behavior

#rnx

diff --git a/Engine/Plugins/Media/ElectraCodecs/Source/ElectraDecoders/Private/Apple/aac/AudioDecoderAAC_Apple.cpp b/Engine/Plugins/Media/ElectraCodecs/Source/ElectraDecoders/Private/Apple/aac/AudioDecoderAAC_Apple.cpp
--- a/Engine/Plugins/Media/ElectraCodecs/Source/ElectraDecoders/Private/Apple/aac/AudioDecoderAAC_Apple.cpp
+++ b/Engine/Plugins/Media/ElectraCodecs/Source/ElectraDecoders/Private/Apple/aac/AudioDecoderAAC_Apple.cpp
@@ -235,7 +235,7 @@ OSStatus FElectraAudioDecoderAAC_Apple::FAudioConverterInstance::AudioToolbox_Co
 	{
 		if (WorkData.InputSize)
 		{
-			if (InOutNumberDataPackets && *InOutNumberDataPackets)
+			if (InOutNumberDataPackets)
 			{
 				*InOutNumberDataPackets = 1;
 				InOutData->mNumberBuffers = 1;
@@ -256,7 +256,10 @@ OSStatus FElectraAudioDecoderAAC_Apple::FAudioConverterInstance::AudioToolbox_Co
 		}
 		else
 		{
-			*InOutNumberDataPackets = 0;
+			if (InOutNumberDataPackets)
+			{
+				*InOutNumberDataPackets = 0;
+			}
 			if (OutDataPacketDescription)
 			{
 				*OutDataPacketDescription = &PacketDescr;
