ElectraDecoders: Creating a global ref to the passed decoder surface when initializing the decoder.

#rnx

diff --git a/Engine/Plugins/Media/ElectraCodecs/Source/ElectraDecoders/Private/Android/h264/VideoDecoderH264_JavaWrapper_Android.cpp b/Engine/Plugins/Media/ElectraCodecs/Source/ElectraDecoders/Private/Android/h264/VideoDecoderH264_JavaWrapper_Android.cpp
--- a/Engine/Plugins/Media/ElectraCodecs/Source/ElectraDecoders/Private/Android/h264/VideoDecoderH264_JavaWrapper_Android.cpp
+++ b/Engine/Plugins/Media/ElectraCodecs/Source/ElectraDecoders/Private/Android/h264/VideoDecoderH264_JavaWrapper_Android.cpp
@@ -466,12 +466,26 @@ int32 FElectraH264VideoDecoderAndroidJava::InitializeDecoder(const FCreateParame
 		}
 
 		// Pass along decoder output surface
-		JEnv->SetObjectField(CreateParams, FCreateParameters_VideoCodecSurface, InCreateParams.VideoCodecSurface);
+		jobject GlobalSurfaceRef = nullptr;
+		if (InCreateParams.VideoCodecSurface)
+		{
+			// Make sure we pass a valid global ref.
+			GlobalSurfaceRef = JEnv->NewGlobalRef(InCreateParams.VideoCodecSurface);
+			if (GlobalSurfaceRef)
+			{
+				JEnv->SetObjectField(CreateParams, FCreateParameters_VideoCodecSurface, GlobalSurfaceRef);
+			}
+		}
 		JEnv->SetBooleanField(CreateParams, FCreateParameters_bSurfaceIsView, InCreateParams.bSurfaceIsView);
 
 		// Create and initialize a decoder instance.
 		result = CallMethodNoVerify<int>(ConfigureDecoderFN, CreateParams);
 		JEnv->DeleteLocalRef(CreateParams);
+		if (GlobalSurfaceRef)
+		{
+			JEnv->DeleteGlobalRef(GlobalSurfaceRef);
+		}
+
 		GClearException(JEnv);
 		if (result != 0)
 		{
diff --git a/Engine/Plugins/Media/ElectraCodecs/Source/ElectraDecoders/Private/Android/h265/VideoDecoderH265_JavaWrapper_Android.cpp b/Engine/Plugins/Media/ElectraCodecs/Source/ElectraDecoders/Private/Android/h265/VideoDecoderH265_JavaWrapper_Android.cpp
--- a/Engine/Plugins/Media/ElectraCodecs/Source/ElectraDecoders/Private/Android/h265/VideoDecoderH265_JavaWrapper_Android.cpp
+++ b/Engine/Plugins/Media/ElectraCodecs/Source/ElectraDecoders/Private/Android/h265/VideoDecoderH265_JavaWrapper_Android.cpp
@@ -460,12 +460,26 @@ int32 FElectraH265VideoDecoderAndroidJava::InitializeDecoder(const FCreateParame
 		}
 
 		// Pass along decoder output surface
-		JEnv->SetObjectField(CreateParams, FCreateParameters_VideoCodecSurface, InCreateParams.VideoCodecSurface);
+		jobject GlobalSurfaceRef = nullptr;
+		if (InCreateParams.VideoCodecSurface)
+		{
+			// Make sure we pass a valid global ref.
+			GlobalSurfaceRef = JEnv->NewGlobalRef(InCreateParams.VideoCodecSurface);
+			if (GlobalSurfaceRef)
+			{
+				JEnv->SetObjectField(CreateParams, FCreateParameters_VideoCodecSurface, GlobalSurfaceRef);
+			}
+		}
 		JEnv->SetBooleanField(CreateParams, FCreateParameters_bSurfaceIsView, InCreateParams.bSurfaceIsView);
 
 		// Create and initialize a decoder instance.
 		result = CallMethodNoVerify<int>(ConfigureDecoderFN, CreateParams);
 		JEnv->DeleteLocalRef(CreateParams);
+		if (GlobalSurfaceRef)
+		{
+			JEnv->DeleteGlobalRef(GlobalSurfaceRef);
+		}
+
 		GClearException(JEnv);
 		if (result != 0)
 		{
