Remove unnecessary Public & Private include paths for subdirectories within a module

#rnx

diff --git a/Engine/Plugins/Experimental/QuicMessaging/Source/QuicMessaging/QuicMessaging.Build.cs b/Engine/Plugins/Experimental/QuicMessaging/Source/QuicMessaging/QuicMessaging.Build.cs
--- a/Engine/Plugins/Experimental/QuicMessaging/Source/QuicMessaging/QuicMessaging.Build.cs
+++ b/Engine/Plugins/Experimental/QuicMessaging/Source/QuicMessaging/QuicMessaging.Build.cs
@@ -29,7 +29,7 @@ namespace UnrealBuildTool.Rules
 					"MsQuicRuntime"
 				}
 			);
-			
+
 			PrivateIncludePathModuleNames.AddRange(
 				new string[] {
 					"MessagingCommon",
@@ -39,11 +39,9 @@ namespace UnrealBuildTool.Rules
 
 			PrivateIncludePaths.AddRange(
 				new string[] {
-					"QuicMessaging/Private",
 					"QuicMessaging/Private/Transport",
 					"QuicMessaging/Private/Serialization",
 					"QuicMessaging/Public/Shared",
-					"QuicMessaging/Public",
 				}
 			);
 
@@ -52,7 +50,8 @@ namespace UnrealBuildTool.Rules
 				}
 			);
 
-			if (Target.Type == TargetType.Editor) {
+			if (Target.Type == TargetType.Editor)
+			{
 
 				DynamicallyLoadedModuleNames.AddRange(
 					new string[] {
diff --git a/Engine/Plugins/Media/ElectraCodecs/Source/ElectraCodecFactory/ElectraCodecFactory.Build.cs b/Engine/Plugins/Media/ElectraCodecs/Source/ElectraCodecFactory/ElectraCodecFactory.Build.cs
--- a/Engine/Plugins/Media/ElectraCodecs/Source/ElectraCodecFactory/ElectraCodecFactory.Build.cs
+++ b/Engine/Plugins/Media/ElectraCodecs/Source/ElectraCodecFactory/ElectraCodecFactory.Build.cs
@@ -2,42 +2,37 @@
 
 namespace UnrealBuildTool.Rules
 {
-    public class ElectraCodecFactory : ModuleRules
-    {
-        public ElectraCodecFactory(ReadOnlyTargetRules Target) : base(Target)
-        {
+	public class ElectraCodecFactory : ModuleRules
+	{
+		public ElectraCodecFactory(ReadOnlyTargetRules Target) : base(Target)
+		{
 			bLegalToDistributeObjectCode = true;
 
-            bool bSupportedPlatform = IsSupportedPlatform(Target);
+			bool bSupportedPlatform = IsSupportedPlatform(Target);
 
-            PrivateDependencyModuleNames.AddRange(
-                new string[] {
+			PrivateDependencyModuleNames.AddRange(
+				new string[] {
 					"ElectraDecoders"
-                });
+				});
 
-            PrivateIncludePaths.AddRange(
-                new string[] {
-                    "ElectraCodecFactory/Private",
-                });
+			PublicDependencyModuleNames.AddRange(
+				new string[] {
+					"Core",
+				});
 
-            PublicDependencyModuleNames.AddRange(
-                new string[] {
-                    "Core",
-                });
+			if (Target.Platform.IsInGroup(UnrealPlatformGroup.Windows))
+			{
+				PublicDefinitions.Add("_CRT_SECURE_NO_WARNINGS");
+			}
 
-            if (Target.Platform.IsInGroup(UnrealPlatformGroup.Windows))
-            {
-                PublicDefinitions.Add("_CRT_SECURE_NO_WARNINGS");
-            }
-
-            if (bSupportedPlatform)
-            {
+			if (bSupportedPlatform)
+			{
 				// ...
-            }
-        }
+			}
+		}
 
 		protected virtual bool IsSupportedPlatform(ReadOnlyTargetRules Target)
-        {
+		{
 			return Target.IsInPlatformGroup(UnrealPlatformGroup.Windows)
 				|| Target.IsInPlatformGroup(UnrealPlatformGroup.Android)
 				|| Target.IsInPlatformGroup(UnrealPlatformGroup.Unix)
@@ -45,5 +40,5 @@ namespace UnrealBuildTool.Rules
 				|| false
 				;
 		}
-    }
+	}
 }
diff --git a/Engine/Plugins/VirtualProduction/LiveLinkXR/Source/LiveLinkXR/LiveLinkXR.Build.cs b/Engine/Plugins/VirtualProduction/LiveLinkXR/Source/LiveLinkXR/LiveLinkXR.Build.cs
--- a/Engine/Plugins/VirtualProduction/LiveLinkXR/Source/LiveLinkXR/LiveLinkXR.Build.cs
+++ b/Engine/Plugins/VirtualProduction/LiveLinkXR/Source/LiveLinkXR/LiveLinkXR.Build.cs
@@ -11,7 +11,6 @@ public class LiveLinkXR : ModuleRules
 		PrivateIncludePaths.AddRange(
 			new string[]
 			{
-				"LiveLinkXR/Private",
 				"LiveLinkXROpenXRExt/Private",
 			}
 		);
diff --git a/Engine/Plugins/VirtualProduction/Takes/Source/CacheTrackRecorder/CacheTrackRecorder.Build.cs b/Engine/Plugins/VirtualProduction/Takes/Source/CacheTrackRecorder/CacheTrackRecorder.Build.cs
--- a/Engine/Plugins/VirtualProduction/Takes/Source/CacheTrackRecorder/CacheTrackRecorder.Build.cs
+++ b/Engine/Plugins/VirtualProduction/Takes/Source/CacheTrackRecorder/CacheTrackRecorder.Build.cs
@@ -6,8 +6,6 @@ public class CacheTrackRecorder : ModuleRules
 {
 	public CacheTrackRecorder(ReadOnlyTargetRules Target) : base(Target)
 	{
-		PublicIncludePaths.Add(ModuleDirectory + "/Public");
-		
 		PrivateDependencyModuleNames.AddRange(
 			new[] {
 				"Core",
@@ -27,9 +25,9 @@ public class CacheTrackRecorder : ModuleRules
 
 		PublicDependencyModuleNames.AddRange(
 			new[] {
-                "TakeTrackRecorders",
-                "Sequencer",
-            }
-        );
-    }
+				"TakeTrackRecorders",
+				"Sequencer",
+			}
+		);
+	}
 }
diff --git a/Engine/Source/Editor/SparseVolumeTexture/SparseVolumeTexture.Build.cs b/Engine/Source/Editor/SparseVolumeTexture/SparseVolumeTexture.Build.cs
--- a/Engine/Source/Editor/SparseVolumeTexture/SparseVolumeTexture.Build.cs
+++ b/Engine/Source/Editor/SparseVolumeTexture/SparseVolumeTexture.Build.cs
@@ -6,9 +6,6 @@ public class SparseVolumeTexture : ModuleRules
 {
 	public SparseVolumeTexture(ReadOnlyTargetRules Target) : base(Target)
 	{
-		PrivateIncludePaths.Add(ModuleDirectory + "/Private");
-		PublicIncludePaths.Add(ModuleDirectory + "/Public");
-
 		PublicDependencyModuleNames.AddRange(
 			new string[]
 			{
@@ -27,7 +24,7 @@ public class SparseVolumeTexture : ModuleRules
 		PrivateIncludePathModuleNames.AddRange(
 			new string[] {
 				"AssetRegistry",
-                "AssetTools",
+				"AssetTools",
 				"EditorWidgets",
 			}
 		);
@@ -36,15 +33,15 @@ public class SparseVolumeTexture : ModuleRules
 			new string[] {
 				"RHI",
 				"CoreUObject",
-                "InputCore",
-                "Settings",
+				"InputCore",
+				"Settings",
 			}
 		);
 
 		if (Target.bBuildEditor)
 		{
 			PrivateDependencyModuleNames.AddRange(
-				new string[] 
+				new string[]
 				{
 					"MainFrame",
 					"EditorFramework",
