ElectraCodecs: Fixed maximum decodable framerate on Windows

#rnx

diff --git a/Engine/Plugins/Media/ElectraCodecs/Source/ElectraDecoders/Private/Windows/h264/H264_VideoDecoder_Windows.cpp b/Engine/Plugins/Media/ElectraCodecs/Source/ElectraDecoders/Private/Windows/h264/H264_VideoDecoder_Windows.cpp
--- a/Engine/Plugins/Media/ElectraCodecs/Source/ElectraDecoders/Private/Windows/h264/H264_VideoDecoder_Windows.cpp
+++ b/Engine/Plugins/Media/ElectraCodecs/Source/ElectraDecoders/Private/Windows/h264/H264_VideoDecoder_Windows.cpp
@@ -191,11 +191,11 @@ void IElectraVideoDecoderH264_DX::PlatformGetSupportedConfigurations(TArray<FSup
 	{
 		IElectraVideoDecoderH264_DX_Platform::DecoderConfigurations.Empty();
 		// Baseline
-		IElectraVideoDecoderH264_DX_Platform::DecoderConfigurations.Emplace(IElectraVideoDecoderH264_DX::FSupportedConfiguration(66, 52, 60, 3840, 2160, 0 /*1280*720 / 256*/));
+		IElectraVideoDecoderH264_DX_Platform::DecoderConfigurations.Emplace(IElectraVideoDecoderH264_DX::FSupportedConfiguration(66, 52, 120, 3840, 2160, 0 /*1280*720 / 256*/));
 		// Main
-		IElectraVideoDecoderH264_DX_Platform::DecoderConfigurations.Emplace(IElectraVideoDecoderH264_DX::FSupportedConfiguration(77, 52, 60, 3840, 2160, 0 /*1280*720 / 256*/));
+		IElectraVideoDecoderH264_DX_Platform::DecoderConfigurations.Emplace(IElectraVideoDecoderH264_DX::FSupportedConfiguration(77, 52, 120, 3840, 2160, 0 /*1280*720 / 256*/));
 		// High
-		IElectraVideoDecoderH264_DX_Platform::DecoderConfigurations.Emplace(IElectraVideoDecoderH264_DX::FSupportedConfiguration(100, 52, 60, 3840, 2160, 0 /*1280*720 / 256*/));
+		IElectraVideoDecoderH264_DX_Platform::DecoderConfigurations.Emplace(IElectraVideoDecoderH264_DX::FSupportedConfiguration(100, 52, 120, 3840, 2160, 0 /*1280*720 / 256*/));
 
 		IElectraVideoDecoderH264_DX_Platform::bDecoderConfigurationsDirty = false;
 	}
diff --git a/Engine/Plugins/Media/ElectraCodecs/Source/ElectraDecoders/Private/Windows/h265/H265_VideoDecoder_Windows.cpp b/Engine/Plugins/Media/ElectraCodecs/Source/ElectraDecoders/Private/Windows/h265/H265_VideoDecoder_Windows.cpp
--- a/Engine/Plugins/Media/ElectraCodecs/Source/ElectraDecoders/Private/Windows/h265/H265_VideoDecoder_Windows.cpp
+++ b/Engine/Plugins/Media/ElectraCodecs/Source/ElectraDecoders/Private/Windows/h265/H265_VideoDecoder_Windows.cpp
@@ -179,9 +179,9 @@ void IElectraVideoDecoderH265_DX::PlatformGetSupportedConfigurations(TArray<FSup
 		IElectraVideoDecoderH265_DX_Platform::DecoderConfigurations.Empty();
 
 		// Main
-		IElectraVideoDecoderH265_DX_Platform::DecoderConfigurations.Emplace(IElectraVideoDecoderH265_DX::FSupportedConfiguration(0, 1, 0, 153, 60, 4096, 2304, 0));
+		IElectraVideoDecoderH265_DX_Platform::DecoderConfigurations.Emplace(IElectraVideoDecoderH265_DX::FSupportedConfiguration(0, 1, 0, 153, 120, 4096, 2304, 0));
 		// Main10
-		IElectraVideoDecoderH265_DX_Platform::DecoderConfigurations.Emplace(IElectraVideoDecoderH265_DX::FSupportedConfiguration(0, 2, 0, 153, 60, 4096, 2304, 0));
+		IElectraVideoDecoderH265_DX_Platform::DecoderConfigurations.Emplace(IElectraVideoDecoderH265_DX::FSupportedConfiguration(0, 2, 0, 153, 120, 4096, 2304, 0));
 
 		IElectraVideoDecoderH265_DX_Platform::bDecoderConfigurationsDirty = false;
 	}
