ElectraDecoders: Added support of VP8 and VP9 for Android

#rnx

diff --git a/Engine/Plugins/Media/ElectraCodecs/Source/ElectraDecoders/ElectraDecoders_UPL.xml b/Engine/Plugins/Media/ElectraCodecs/Source/ElectraDecoders/ElectraDecoders_UPL.xml
--- a/Engine/Plugins/Media/ElectraCodecs/Source/ElectraDecoders/ElectraDecoders_UPL.xml
+++ b/Engine/Plugins/Media/ElectraCodecs/Source/ElectraDecoders/ElectraDecoders_UPL.xml
@@ -56,6 +56,21 @@
 -keep class com.epicgames.unreal.ElectraDecoderVideoH265$FOutputBufferInfo {
 	public *;
 }
+-keep class com.epicgames.unreal.ElectraDecoderVideoVPx {
+	public *;
+}
+-keep class com.epicgames.unreal.ElectraDecoderVideoVPx$FCreateParameters {
+	public *;
+}
+-keep class com.epicgames.unreal.ElectraDecoderVideoVPx$FDecoderInformation {
+	public *;
+}
+-keep class com.epicgames.unreal.ElectraDecoderVideoVPx$FOutputFormatInfo {
+	public *;
+}
+-keep class com.epicgames.unreal.ElectraDecoderVideoVPx$FOutputBufferInfo {
+	public *;
+}
 -keep class com.epicgames.unreal.ElectraDecoderAudioAAC {
 	public *;
 }
diff --git a/Engine/Plugins/Media/ElectraCodecs/Source/ElectraDecoders/Private/Android/AndroidPlatformElectraDecoders.cpp b/Engine/Plugins/Media/ElectraCodecs/Source/ElectraDecoders/Private/Android/AndroidPlatformElectraDecoders.cpp
--- a/Engine/Plugins/Media/ElectraCodecs/Source/ElectraDecoders/Private/Android/AndroidPlatformElectraDecoders.cpp
+++ b/Engine/Plugins/Media/ElectraCodecs/Source/ElectraDecoders/Private/Android/AndroidPlatformElectraDecoders.cpp
@@ -5,6 +5,7 @@
 
 #include "h264/H264_VideoDecoder_Android.h"
 #include "h265/H265_VideoDecoder_Android.h"
+#include "vpx/VPx_VideoDecoder_Android.h"
 #include "aac/AAC_AudioDecoder_Android.h"
 
 namespace PlatformElectraDecodersAndroid
@@ -40,6 +41,8 @@ void FPlatformElectraDecodersAndroid::RegisterWithCodecFactory(IElectraCodecRegi
 		RegisterFactory(FH264VideoDecoderAndroid::CreateFactory());
 		// H.265 video decoder
 		RegisterFactory(FH265VideoDecoderAndroid::CreateFactory());
+		// VP8/VP9 video decoder
+		RegisterFactory(FVPxVideoDecoderAndroid::CreateFactory());
 		// AAC audio decoder
 		RegisterFactory(FAACAudioDecoderAndroid::CreateFactory());
 
diff --git a/Engine/Plugins/Media/ElectraCodecs/Source/ElectraDecoders/Private/Android/Java/src/com/epicgames/unreal/ElectraDecoderVideoVPx.java b/Engine/Plugins/Media/ElectraCodecs/Source/ElectraDecoders/Private/Android/Java/src/com/epicgames/unreal/ElectraDecoderVideoVPx.java
new file mode 100644
--- /dev/null
+++ b/Engine/Plugins/Media/ElectraCodecs/Source/ElectraDecoders/Private/Android/Java/src/com/epicgames/unreal/ElectraDecoderVideoVPx.java
@@ -0,0 +1,829 @@
+// Copyright Epic Games, Inc. All Rights Reserved.
+
+package com.epicgames.unreal;
+
+import android.annotation.SuppressLint;
+import android.annotation.TargetApi;
+
+import android.os.Build;
+import java.io.IOException;
+import java.util.ArrayList;
+import android.util.Log;
+import android.util.SparseArray;
+
+import android.view.Surface;
+import android.media.AudioFormat;
+import android.media.MediaCodecList;
+import android.media.MediaCodec;
+import android.media.MediaCodecInfo;
+import android.media.MediaCodecInfo.CodecCapabilities;
+import android.media.MediaCodecInfo.CodecProfileLevel;
+import android.media.MediaFormat;
+import java.nio.ByteBuffer;
+
+import com.epicgames.unreal.ElectraAndroidCodecQuirks;
+
+
+
+/**
+ * VP8/VP9 video decoder
+ *
+ */
+public class ElectraDecoderVideoVPx
+{
+	private static final String TAG = "[ElectraDecoders]";
+	/**
+	 * Parameters for creating a new decoder instance.
+	 * NOTE: Make sure this is a static inner class so it can be
+	 *       created from C++ through JNI.
+	 */
+	public static class FCreateParameters
+	{
+		public int MaxWidth;
+		public int MaxHeight;
+		public int MaxFPS;
+		public int Width;
+		public int Height;
+		public boolean bNeedSecure;
+		public boolean bNeedTunneling;
+		public byte[] CSD0;
+		public int NativeDecoderID;
+		public android.view.Surface VideoCodecSurface;
+		public boolean bSurfaceIsView;
+		public FCreateParameters()
+		{
+			MaxWidth = 0;
+			MaxHeight = 0;
+			Width = 0;
+			Height = 0;
+			MaxFPS = 0;
+			bNeedSecure = false;
+			bNeedTunneling = false;
+			NativeDecoderID = 0;
+			VideoCodecSurface = null;
+		}
+	}
+
+	public static class FDecoderInformation
+	{
+		public int ApiLevel = 0;
+		public boolean bIsAdaptive = false;
+		public boolean bCanUse_SetOutputSurface = false;
+	}
+
+	public class FOutputFormatInfo
+	{
+		public int Width = 0;
+		public int Height = 0;
+		public int CropTop = 0;
+		public int CropBottom = 0;
+		public int CropLeft = 0;
+		public int CropRight = 0;
+		public int Stride = 0;
+		public int SliceHeight = 0;
+		public int ColorFormat = 0;
+	}
+
+	public class FOutputBufferInfo
+	{
+		public int BufferIndex = -1;
+		public long PresentationTimestamp = -1;
+		public int Size = 0;
+		public boolean bIsEOS = false;
+		public boolean bIsConfig = false;
+	}
+
+	private static class FCodecInformation
+	{
+		MediaCodecInfo Info;
+		MediaCodecInfo.CodecCapabilities Caps;
+	}
+
+
+	private static class FVPxDecoderInfo
+	{
+		private static final FVPxDecoderInfo Singleton = new FVPxDecoderInfo();
+		private FCodecInformation CodecInfoVP8;
+		private FCodecInformation CodecInfoVP9;
+		private FVPxDecoderInfo()
+		{
+			CodecInfoVP8 = selectCodec(MediaFormat.MIMETYPE_VIDEO_VP8);
+			CodecInfoVP9 = selectCodec(MediaFormat.MIMETYPE_VIDEO_VP9);
+		}
+
+		private static final boolean DoesDecoderSupportsLowLatency(MediaCodecInfo decoderInfo, String mimeType)
+		{
+			if (Build.VERSION.SDK_INT >= 30)
+			{
+				try
+				{
+					if (decoderInfo.getCapabilitiesForType(mimeType).isFeatureSupported("low-latency"))
+					{
+						return true;
+					}
+				}
+				catch (Exception e)
+				{
+				}
+			}
+			return false;
+		}
+
+		private static final FCodecInformation selectCodec(String mimeType)
+		{
+			/*
+			If we knew what to put into the media format exactly we could perhaps use findDecoderForFormat().
+				@see: https://developer.android.com/reference/android/media/MediaCodec.html#creation
+			Since we do not have an extractor to get the track format from, after all we are decoding an elementary stream here,
+			we have no idea what would normally be found and what properties are _required_ in the format.
+
+				MediaFormat testFmt = new MediaFormat();
+				String thisOne = cl.findDecoderForFormat(testFmt);
+			*/
+
+			MediaCodecList cl = new MediaCodecList(MediaCodecList.ALL_CODECS);
+			MediaCodecInfo[] codecInfos = cl.getCodecInfos();
+			int numCodecs = codecInfos.length;
+			for(int pass=0; pass<2; ++pass)
+			{
+				for(int i=0; i<numCodecs; ++i)
+				{
+					if (codecInfos[i].isEncoder())
+					{
+						continue;
+					}
+
+					// If this is an alias for another decoder (starting with API level 29, aka "Q") we skip it.
+					if (Build.VERSION.SDK_INT >= 29)
+					{
+						if (codecInfos[i].isAlias())
+						{
+							continue;
+						}
+					}
+
+					String[] types = codecInfos[i].getSupportedTypes();
+					for(int j=0; j<types.length; ++j)
+					{
+						if (types[j].equalsIgnoreCase(mimeType))
+						{
+							// Apparently decoders may exists twice, one that doesn't support low latency and another that does.
+							// And the one that does is not placed in front of the one that doesn't, so we need to check twice.
+							boolean bSupportsLowLatency = DoesDecoderSupportsLowLatency(codecInfos[i], mimeType);
+
+							//Log.w(TAG, "["+pass+"] ElectraDecoderVideoVPx: " + codecInfos[i].getCanonicalName() + "; ll="+bSupportsLowLatency);
+
+							if (pass==0 && !bSupportsLowLatency)
+							{
+								continue;
+							}
+
+							// Rumor has it that the codec list is sorted from "best" to "worst".
+							// Assuming this to be true we just return the first one matching the MIME type.
+							FCodecInformation Info = new FCodecInformation();
+							Info.Info = codecInfos[i];
+							Info.Caps = codecInfos[i].getCapabilitiesForType(types[j]);
+							return Info;
+						}
+					}
+				}
+			}
+			return null;
+		}
+
+		public static FCodecInformation GetDecoderInfo(boolean bForVP9)
+		{
+			if (bForVP9 == true)
+			{
+				return Singleton.CodecInfoVP9;
+			}
+			else
+			{
+				return Singleton.CodecInfoVP8;
+			}
+		}
+	}
+
+
+	private FCodecInformation PlatformCodecInfo;
+	private ElectraAndroidCodecQuirks.FDecoderQuirks PlatformCodecQuirks;
+
+	private FCreateParameters CreationParameters;
+	private FDecoderInformation DecoderInformation;
+	private MediaCodec DecoderHandle;
+	private boolean bIsVP9;
+	private boolean bIsInitialized;
+
+	private SparseArray<MediaCodec.BufferInfo> OutputBufferInfos;
+	private FOutputFormatInfo CurrentOutputFormatInfo;
+	private boolean bDecoderHasSurface;
+
+
+	public ElectraDecoderVideoVPx()
+	{
+		PlatformCodecInfo = null;
+		PlatformCodecQuirks =  ElectraAndroidCodecQuirks.GetDecoderQuirks();
+		CreationParameters = null;
+		DecoderHandle = null;
+		OutputBufferInfos = new SparseArray<>();
+		CurrentOutputFormatInfo = new FOutputFormatInfo();
+		bDecoderHasSurface = false;
+		bIsVP9 = true;
+		bIsInitialized = false;
+
+		// Create the initial decoder information, which at this point contains the quirks only.
+		DecoderInformation = new FDecoderInformation();
+		SetupDecoderQuirks();
+	}
+
+	private void SetupDecoderQuirks()
+	{
+		if (DecoderInformation != null && PlatformCodecQuirks != null)
+		{
+			DecoderInformation.ApiLevel = Build.VERSION.SDK_INT;
+			DecoderInformation.bCanUse_SetOutputSurface = PlatformCodecQuirks.bCanUse_SetOutputSurface;
+		}
+	}
+
+	public int CreateDecoder(boolean bCreateVP9)
+	{
+		// Clear out everything in case this is called repeatedly.
+		DecoderHandle = null;
+		DecoderInformation = null;
+		bIsVP9 = bCreateVP9;
+		bIsInitialized = false;
+		OutputBufferInfos = new SparseArray<>();
+		CurrentOutputFormatInfo = new FOutputFormatInfo();
+		PlatformCodecInfo = FVPxDecoderInfo.GetDecoderInfo(bIsVP9);
+
+		try
+		{
+			if (PlatformCodecInfo != null)
+			{
+				String DecoderName = PlatformCodecInfo.Info.getName();
+
+				// This is how we get information on the decoder's capabilities.
+				// For now we do not use those. We merely assume the decoder is capable of what we want.
+				// Caution must be exercised with profile/level values as they are defined with other values in Android than they are by the standard!
+				// Eg. High profile per the standard is identified by the value 100 while Android defines it as 8 (MediaCodecInfo.CodecProfileLevel.AVCProfileHigh)
+				boolean bAdaptive   						= PlatformCodecInfo.Caps.isFeatureSupported(MediaCodecInfo.CodecCapabilities.FEATURE_AdaptivePlayback);
+				//boolean bSecure 							= PlatformCodecInfo.Caps.isFeatureSupported(MediaCodecInfo.CodecCapabilities.FEATURE_SecurePlayback);
+				//boolean bTunneled   						= PlatformCodecInfo.Caps.isFeatureSupported(MediaCodecInfo.CodecCapabilities.FEATURE_TunneledPlayback);
+				//int [] ColorFmts							= PlatformCodecInfo.Caps.colorFormats;
+				//MediaCodecInfo.CodecProfileLevel [] ProfLev = PlatformCodecInfo.Caps.profileLevels;
+
+				// Create the decoder
+				DecoderHandle = MediaCodec.createByCodecName(DecoderName);
+				if (DecoderHandle != null)
+				{
+					// Remember the decoder capabilities
+					DecoderInformation = new FDecoderInformation();
+					DecoderInformation.bIsAdaptive = bAdaptive;
+					SetupDecoderQuirks();
+				}
+				else
+				{
+					Log.w(TAG, "ElectraDecoderVideoVPx: Failed to create decoder by name");
+					return 1;
+				}
+			}
+			else
+			{
+				// No decoder found.
+				Log.w(TAG, "ElectraDecoderVideoVPx: No suitable decoder found");
+				return 1;
+			}
+		}
+		catch(Exception e)
+		{
+			Log.w(TAG, "ElectraDecoderVideoVPx: Failed to create decoder");
+			e.printStackTrace();
+			return 1;
+		}
+		return 0;
+	}
+
+
+	public int ConfigureDecoder(FCreateParameters InCreateParams)
+	{
+		if (DecoderHandle == null || DecoderInformation == null || PlatformCodecInfo == null)
+		{
+			Log.w(TAG, "ElectraDecoderVideoVPx: No decoder instance to configure has been created yet");
+			return 1;
+		}
+
+		try
+		{
+			CreationParameters = InCreateParams;
+
+			// Create decoder configuration object
+			int MaxWidth  =  InCreateParams.MaxWidth > 0 ? InCreateParams.MaxWidth  : 1920;
+			int MaxHeight = InCreateParams.MaxHeight > 0 ? InCreateParams.MaxHeight : 1088;
+			int Width  = InCreateParams.Width > 0  ? InCreateParams.Width : MaxWidth;
+			int Height = InCreateParams.Height > 0 ? InCreateParams.Height : MaxHeight;
+			MediaFormat Format = MediaFormat.createVideoFormat(bIsVP9 ? MediaFormat.MIMETYPE_VIDEO_VP9 : MediaFormat.MIMETYPE_VIDEO_VP8, Width, Height);
+			// Add additional keys
+			if (DecoderInformation.bIsAdaptive)
+			{
+				// See: https://developer.android.com/reference/android/media/MediaFormat.html
+				Format.setInteger(MediaFormat.KEY_MAX_WIDTH, MaxWidth);
+				Format.setInteger(MediaFormat.KEY_MAX_HEIGHT, MaxHeight);
+			}
+			else
+			{
+				// Non-adaptive decoders need the CSD set in the configuration.
+				if (InCreateParams.CSD0 != null && InCreateParams.CSD0.length != 0)
+				{
+					ByteBuffer Csd0 = ByteBuffer.wrap(InCreateParams.CSD0);
+					Format.setByteBuffer("csd-0", Csd0);
+				}
+			}
+			// Ask for low-latency.
+			Format.setInteger("low-latency", 1);
+
+			// Configure the decoder
+			android.view.Surface OutputSurface = null;
+			OutputSurface = InCreateParams.VideoCodecSurface;
+			bDecoderHasSurface = (OutputSurface != null);
+			DecoderHandle.configure(Format, OutputSurface, null, 0);
+			bIsInitialized = true;
+
+			// We are not storing any outside surface reference (we only hand it through to the decoder)
+			CreationParameters.VideoCodecSurface = null;
+		}
+		catch(Exception e)
+		{
+			Log.w(TAG, "ElectraDecoderVideoVPx: Failed to configure the decoder");
+			e.printStackTrace();
+			return 1;
+		}
+		return 0;
+	}
+
+
+	@TargetApi(23)
+	public int SetOutputSurface(android.view.Surface NewVideoCodecSurface)
+	{
+		if (Build.VERSION.SDK_INT >= 23 && bIsInitialized && NewVideoCodecSurface != null)
+		{
+			try
+			{
+				bDecoderHasSurface = true;
+				DecoderHandle.setOutputSurface(NewVideoCodecSurface);
+			}
+			catch(Exception e)
+			{
+				Log.w(TAG, "ElectraDecoderVideoVPx: Failed to set decoder output surface");
+				e.printStackTrace();
+				return 1;
+			}
+			return 0;
+		}
+		else
+		{
+			return 1;
+		}
+	}
+
+
+	public int ReleaseDecoder()
+	{
+		if (bIsInitialized)
+		{
+			bIsInitialized = false;
+
+			try
+			{
+				// Release the decoder
+				DecoderHandle.release();
+			}
+			catch(Exception e)
+			{
+				Log.w(TAG, "ElectraDecoderVideoVPx: Failed to release decoder");
+				e.printStackTrace();
+				return 1;
+			}
+			return 0;
+		}
+		else
+		{
+			return 1;
+		}
+	}
+
+
+	public int release()
+	{
+		ReleaseDecoder();
+		return 0;
+	}
+
+
+	public int Start()
+	{
+		if (bIsInitialized)
+		{
+			try
+			{
+				// Start the decoder
+				DecoderHandle.start();
+			}
+			catch(Exception e)
+			{
+				Log.w(TAG, "ElectraDecoderVideoVPx: Failed to start decoder");
+				e.printStackTrace();
+				return 1;
+			}
+			return 0;
+		}
+		else
+		{
+			return 1;
+		}
+	}
+
+
+	public int Stop()
+	{
+		if (bIsInitialized)
+		{
+			try
+			{
+				// Stop the decoder
+				DecoderHandle.stop();
+			}
+			catch(Exception e)
+			{
+				Log.w(TAG, "ElectraDecoderVideoVPx: Failed to stop decoder");
+				e.printStackTrace();
+				return 1;
+			}
+			return 0;
+		}
+		else
+		{
+			return 1;
+		}
+	}
+
+
+	public int Flush()
+	{
+		if (bIsInitialized)
+		{
+			try
+			{
+				// Flush the decoder
+				DecoderHandle.flush();
+			}
+			catch(Exception e)
+			{
+				Log.w(TAG, "ElectraDecoderVideoVPx: Failed to flush decoder");
+				e.printStackTrace();
+				return 1;
+			}
+
+			// Clear the current infos by creating new empty structures.
+			OutputBufferInfos = new SparseArray<>();
+			// NOTE: Flushing does not invalidate the current output configuration. Do not clear it.
+			// CurrentOutputFormatInfo = new FOutputFormatInfo();
+			return 0;
+		}
+		else
+		{
+			return 1;
+		}
+	}
+
+
+	public int Reset()
+	{
+		if (bIsInitialized)
+		{
+			try
+			{
+				// Reset the decoder
+				DecoderHandle.reset();
+			}
+			catch(Exception e)
+			{
+				Log.w(TAG, "ElectraDecoderVideoVPx: Failed to reset decoder");
+				e.printStackTrace();
+				return 1;
+			}
+			return 0;
+		}
+		else
+		{
+			return 1;
+		}
+	}
+
+
+	public int DequeueInputBuffer(int InTimeoutUsec)
+	{
+		if (bIsInitialized)
+		{
+			int BufferIndex = -1;
+			try
+			{
+				BufferIndex = DecoderHandle.dequeueInputBuffer(InTimeoutUsec);
+			}
+			catch(Exception e)
+			{
+				Log.w(TAG, "ElectraDecoderVideoVPx: Failed to dequeue input buffer");
+				e.printStackTrace();
+				return -10000;
+			}
+			return BufferIndex;
+		}
+		else
+		{
+			return -10001;
+		}
+	}
+
+
+	public int QueueInputBuffer(int InBufferIndex, long InTimestampUSec, byte[] InAccessUnit)
+	{
+		if (bIsInitialized)
+		{
+			if (InBufferIndex >= 0)
+			{
+				try
+				{
+					int SizeOfAU = InAccessUnit.length;
+					ByteBuffer Buffer = DecoderHandle.getInputBuffer(InBufferIndex);
+					Buffer.put(InAccessUnit);
+					DecoderHandle.queueInputBuffer(InBufferIndex, 0, SizeOfAU, InTimestampUSec, 0);
+				}
+				catch(Exception e)
+				{
+					Log.w(TAG, "ElectraDecoderVideoVPx: Failed to queue input buffer");
+					e.printStackTrace();
+					return -10000;
+				}
+				return 0;
+			}
+			else
+			{
+				return -10002;
+			}
+		}
+		else
+		{
+			return -10001;
+		}
+	}
+
+
+	public int QueueCSDInputBuffer(int InBufferIndex, long InTimestampUSec, byte[] InCSD)
+	{
+		if (bIsInitialized)
+		{
+			if (InBufferIndex >= 0)
+			{
+				try
+				{
+					int SizeOfCSD = InCSD.length;
+					ByteBuffer Buffer = DecoderHandle.getInputBuffer(InBufferIndex);
+					Buffer.put(InCSD);
+					DecoderHandle.queueInputBuffer(InBufferIndex, 0, SizeOfCSD, InTimestampUSec, MediaCodec.BUFFER_FLAG_CODEC_CONFIG);
+				}
+				catch(Exception e)
+				{
+					Log.w(TAG, "ElectraDecoderVideoVPx: Failed to queue CSD input buffer");
+					e.printStackTrace();
+					return -10000;
+				}
+				return 0;
+			}
+			else
+			{
+				return -10002;
+			}
+		}
+		else
+		{
+			return -10001;
+		}
+	}
+
+
+	public int QueueEOSInputBuffer(int InBufferIndex, long InTimestampUSec)
+	{
+		if (bIsInitialized)
+		{
+			if (InBufferIndex >= 0)
+			{
+				try
+				{
+					ByteBuffer Buffer = DecoderHandle.getInputBuffer(InBufferIndex);
+					Buffer.clear();
+					DecoderHandle.queueInputBuffer(InBufferIndex, 0, 0, InTimestampUSec, MediaCodec.BUFFER_FLAG_END_OF_STREAM);
+				}
+				catch(Exception e)
+				{
+					Log.w(TAG, "ElectraDecoderVideoVPx: Failed to queue EOS input buffer");
+					e.printStackTrace();
+					return -10000;
+				}
+				return 0;
+			}
+			else
+			{
+				return -10002;
+			}
+		}
+		else
+		{
+			return -10001;
+		}
+	}
+
+
+	public FDecoderInformation GetDecoderInformation()
+	{
+		return DecoderInformation;
+	}
+
+
+	private FOutputFormatInfo GetOutputFormatInfoFromFormat(MediaFormat Format)
+	{
+		FOutputFormatInfo NewFormatInfo = new FOutputFormatInfo();
+		if (Format.containsKey(MediaFormat.KEY_WIDTH))
+		{
+			NewFormatInfo.Width = Format.getInteger(MediaFormat.KEY_WIDTH);
+		}
+		if (Format.containsKey(MediaFormat.KEY_HEIGHT))
+		{
+			NewFormatInfo.Height = Format.getInteger(MediaFormat.KEY_HEIGHT);
+		}
+		if (Format.containsKey("slice-height"))	// MediaFormat.KEY_SLICE_HEIGHT  with API level 23 only...
+		{
+			NewFormatInfo.SliceHeight = Format.getInteger("slice-height");
+		}
+		if (Format.containsKey("stride"))	// MediaFormat.KEY_STRIDE  with API level 23 only...
+		{
+			NewFormatInfo.Stride = Format.getInteger("stride");
+		}
+		if (Format.containsKey(MediaFormat.KEY_COLOR_FORMAT))
+		{
+			NewFormatInfo.ColorFormat = Format.getInteger(MediaFormat.KEY_COLOR_FORMAT);
+		}
+		// See: https://developer.android.com/reference/android/media/MediaCodec#accessing-raw-video-bytebuffers-on-older-devices
+		if (Format.containsKey("crop-left") && Format.containsKey("crop-right") && Format.containsKey("crop-top") && Format.containsKey("crop-bottom"))
+		{
+			NewFormatInfo.CropLeft   = Format.getInteger("crop-left");
+			NewFormatInfo.CropRight  = Format.getInteger("crop-right");
+			NewFormatInfo.CropTop    = Format.getInteger("crop-top");
+			NewFormatInfo.CropBottom = Format.getInteger("crop-bottom");
+		}
+		else
+		{
+			NewFormatInfo.CropLeft   = 0;
+			NewFormatInfo.CropTop    = 0;
+			NewFormatInfo.CropRight  = NewFormatInfo.Width - 1;
+			NewFormatInfo.CropBottom = NewFormatInfo.Height - 1;
+		}
+		return NewFormatInfo;
+	}
+
+
+	public FOutputFormatInfo GetOutputFormatInfo(int InOutputBufferIndex)
+	{
+		if (InOutputBufferIndex < 0)
+		{
+			return CurrentOutputFormatInfo;
+		}
+		if (bIsInitialized)
+		{
+			MediaFormat Format = DecoderHandle.getOutputFormat(InOutputBufferIndex);
+			return GetOutputFormatInfoFromFormat(Format);
+		}
+		return null;
+	}
+
+
+	public FOutputBufferInfo DequeueOutputBuffer(int InTimeoutUsec)
+	{
+		if (bIsInitialized)
+		{
+			FOutputBufferInfo OutputBufferInfo = new FOutputBufferInfo();
+			MediaCodec.BufferInfo Info = new MediaCodec.BufferInfo();
+			try
+			{
+				OutputBufferInfo.BufferIndex = DecoderHandle.dequeueOutputBuffer(Info, InTimeoutUsec);
+				// If there is output we need to remember the buffer metadata for later.
+				if (OutputBufferInfo.BufferIndex >= 0)
+				{
+					OutputBufferInfo.PresentationTimestamp = Info.presentationTimeUs;
+					OutputBufferInfo.Size = Info.size;
+					OutputBufferInfo.bIsEOS = (Info.flags & MediaCodec.BUFFER_FLAG_END_OF_STREAM) != 0;
+					OutputBufferInfo.bIsConfig = (Info.flags & MediaCodec.BUFFER_FLAG_CODEC_CONFIG) != 0;
+					OutputBufferInfos.put(OutputBufferInfo.BufferIndex, Info);
+				}
+				else if (OutputBufferInfo.BufferIndex == MediaCodec.INFO_TRY_AGAIN_LATER)				// -1
+				{
+					// Have no output yet.
+				}
+				else if (OutputBufferInfo.BufferIndex == MediaCodec.INFO_OUTPUT_FORMAT_CHANGED)			// -2
+				{
+					MediaFormat Format = DecoderHandle.getOutputFormat();
+					CurrentOutputFormatInfo = GetOutputFormatInfoFromFormat(Format);
+				}
+				else if (OutputBufferInfo.BufferIndex == MediaCodec.INFO_OUTPUT_BUFFERS_CHANGED)		// -3
+				{
+					// This is deprecated with API level 21 and of no real value to us.
+				}
+			}
+			catch(Exception e)
+			{
+				Log.w(TAG, "ElectraDecoderVideoVPx: Failed to dequeue output buffer");
+				e.printStackTrace();
+				return null;
+			}
+			return OutputBufferInfo;
+		}
+		return null;
+	}
+
+
+	public byte[] GetOutputBuffer(int InOutputBufferIndex)
+	{
+		byte[] Data = null;
+		if (bIsInitialized && InOutputBufferIndex >= 0)
+		{
+			MediaCodec.BufferInfo Info = OutputBufferInfos.get(InOutputBufferIndex);
+			if (Info != null)
+			{
+				try
+				{
+					ByteBuffer outBuffer = DecoderHandle.getOutputBuffer(InOutputBufferIndex);
+					Data = new byte[Info.size];
+					outBuffer.get(Data);
+				}
+				catch(Exception e)
+				{
+					Log.w(TAG, "ElectraDecoderVideoVPx: Failed to get output buffer");
+					e.printStackTrace();
+					return null;
+				}
+			}
+		}
+		return Data;
+	}
+
+
+	public int ReleaseOutputBuffer(int InOutputBufferIndex, boolean bRender, long releaseAt)
+	{
+		if (bIsInitialized && InOutputBufferIndex >= 0)
+		{
+			MediaCodec.BufferInfo Info = OutputBufferInfos.get(InOutputBufferIndex);
+			if (Info != null)
+			{
+				OutputBufferInfos.delete(InOutputBufferIndex);
+
+				FOutputFormatInfo OutFmtInfo = GetOutputFormatInfo(InOutputBufferIndex);
+
+				// If the decoder has no output surface we cannot release the buffer to it!
+				bRender &= bDecoderHasSurface;
+
+				try
+				{
+					boolean bDoRender = bRender == true && (Info.size != 0) && ((Info.flags & MediaCodec.BUFFER_FLAG_CODEC_CONFIG) != MediaCodec.BUFFER_FLAG_CODEC_CONFIG);
+					if (releaseAt < 0 || bDoRender == false)
+					{
+						if (bDoRender && CreationParameters.bSurfaceIsView)
+						{
+							// We use a timed output instead of the un-timed one because that one seems to use PTS values gathered from the sample(??) that of course do
+							// not correclate to "realtime within 1s" - and trigger warnings in logcat...
+							DecoderHandle.releaseOutputBuffer(InOutputBufferIndex, System.nanoTime());
+						}
+						else
+						{
+							DecoderHandle.releaseOutputBuffer(InOutputBufferIndex, bDoRender);
+						}
+					}
+					else
+					{
+						DecoderHandle.releaseOutputBuffer(InOutputBufferIndex, releaseAt);
+					}
+				}
+				catch(Exception e)
+				{
+					Log.w(TAG, "ElectraDecoderVideoVPx: Failed to release buffer to surface");
+					e.printStackTrace();
+					return 1;
+				}
+				return 0;
+			}
+			else
+			{
+				return -10000;
+			}
+		}
+		return 1;
+	}
+}
+
diff --git a/Engine/Plugins/Media/ElectraCodecs/Source/ElectraDecoders/Private/Android/vpx/VPx_VideoDecoder_Android.cpp b/Engine/Plugins/Media/ElectraCodecs/Source/ElectraDecoders/Private/Android/vpx/VPx_VideoDecoder_Android.cpp
new file mode 100644
--- /dev/null
+++ b/Engine/Plugins/Media/ElectraCodecs/Source/ElectraDecoders/Private/Android/vpx/VPx_VideoDecoder_Android.cpp
@@ -0,0 +1,56 @@
+// Copyright Epic Games, Inc. All Rights Reserved.
+
+#include "vpx/VPx_VideoDecoder_Android.h"
+#include "vpx/VideoDecoderVPx_Android.h"
+
+#include "ElectraDecodersUtils.h"
+
+
+class FVPxVideoDecoderFactoryAndroid : public IElectraCodecFactory
+{
+public:
+	virtual ~FVPxVideoDecoderFactoryAndroid()
+	{}
+
+	int32 SupportsFormat(const FString& InCodecFormat, bool bInEncoder, const TMap<FString, FVariant>& InOptions) const override
+	{
+		// Encoder? Not supported here!
+		if (bInEncoder)
+		{
+			return 0;
+		}
+
+		if (InCodecFormat.StartsWith(TEXT("vp08"), ESearchCase::IgnoreCase))
+		{
+			return 1;
+		}
+		else if (InCodecFormat.StartsWith(TEXT("vp09"), ESearchCase::IgnoreCase))
+		{
+			ElectraDecodersUtil::FMimeTypeVideoCodecInfo ci;
+			if (ElectraDecodersUtil::ParseCodecVP9(ci, InCodecFormat, ElectraDecodersUtil::GetVariantValueUInt8Array(InOptions, TEXT("$vpcC_box"))))
+			{
+				return 1;
+			}
+		}
+		return 0;
+	}
+
+
+	void GetConfigurationOptions(TMap<FString, FVariant>& OutOptions) const override
+	{
+		IElectraVideoDecoderVPx_Android::GetConfigurationOptions(OutOptions);
+	}
+
+	TSharedPtr<IElectraDecoder, ESPMode::ThreadSafe> CreateDecoderForFormat(const FString& InCodecFormat, const TMap<FString, FVariant>& InOptions, TSharedPtr<IElectraDecoderResourceDelegate, ESPMode::ThreadSafe> InResourceDelegate) override
+	{
+		check(SupportsFormat(InCodecFormat, false, InOptions));
+		return IElectraVideoDecoderVPx_Android::Create(InOptions, InResourceDelegate);
+	}
+
+};
+
+
+TSharedPtr<IElectraCodecFactory, ESPMode::ThreadSafe> FVPxVideoDecoderAndroid::CreateFactory()
+{
+	return MakeShared<FVPxVideoDecoderFactoryAndroid, ESPMode::ThreadSafe>();
+}
diff --git a/Engine/Plugins/Media/ElectraCodecs/Source/ElectraDecoders/Private/Android/vpx/VPx_VideoDecoder_Android.h b/Engine/Plugins/Media/ElectraCodecs/Source/ElectraDecoders/Private/Android/vpx/VPx_VideoDecoder_Android.h
new file mode 100644
--- /dev/null
+++ b/Engine/Plugins/Media/ElectraCodecs/Source/ElectraDecoders/Private/Android/vpx/VPx_VideoDecoder_Android.h
@@ -0,0 +1,14 @@
+// Copyright Epic Games, Inc. All Rights Reserved.
+
+#pragma once
+
+#include "CoreMinimal.h"
+
+#include "IElectraCodecFactory.h"
+
+class FVPxVideoDecoderAndroid
+{
+public:
+	static TSharedPtr<IElectraCodecFactory, ESPMode::ThreadSafe> CreateFactory();
+
+};
diff --git a/Engine/Plugins/Media/ElectraCodecs/Source/ElectraDecoders/Private/Android/vpx/VideoDecoderVPx_Android.cpp b/Engine/Plugins/Media/ElectraCodecs/Source/ElectraDecoders/Private/Android/vpx/VideoDecoderVPx_Android.cpp
new file mode 100644
--- /dev/null
+++ b/Engine/Plugins/Media/ElectraCodecs/Source/ElectraDecoders/Private/Android/vpx/VideoDecoderVPx_Android.cpp
@@ -0,0 +1,1078 @@
+// Copyright Epic Games, Inc. All Rights Reserved.
+
+#include "vpx/VideoDecoderVPx_Android.h"
+#include "DecoderErrors_Android.h"
+#include "ElectraDecodersUtils.h"
+#include "Utils/Google/ElectraUtilsVPxVideo.h"
+
+#include "IElectraDecoderFeaturesAndOptions.h"
+#include "IElectraDecoderOutputVideo.h"
+
+#include "IElectraDecoderResourceDelegate.h"
+
+#include "ElectraDecodersModule.h"
+
+#include <atomic>
+/*********************************************************************************************************************/
+
+#include "vpx/VideoDecoderVPx_JavaWrapper_Android.h"
+
+/*********************************************************************************************************************/
+#define ENABLE_DETAILED_LOG 0
+#if ENABLE_DETAILED_LOG
+#define DETAILLOG UE_LOG
+#else
+#define DETAILLOG(CategoryName, Verbosity, Format, ...) while(0){}
+#endif
+
+#define DESTROY_DECODER_WHEN_FLUSHING 1
+
+/*********************************************************************************************************************/
+
+class FElectraVideoDecoderVPx_Android;
+
+
+class FElectraDecoderDefaultVideoOutputFormatVPx_Android : public IElectraDecoderDefaultVideoOutputFormat
+{
+public:
+	virtual ~FElectraDecoderDefaultVideoOutputFormatVPx_Android()
+	{ }
+};
+
+
+class FElectraVideoDecoderOutputVPx_Android : public IElectraDecoderVideoOutput
+{
+public:
+	virtual ~FElectraVideoDecoderOutputVPx_Android()
+	{
+		ReleaseOutputBuffer();
+	}
+
+	FTimespan GetPTS() const override
+	{ return PTS; }
+	uint64 GetUserValue() const override
+	{ return UserValue; }
+
+	int32 GetWidth() const override
+	{ return Width - Crop.Left - Crop.Right; }
+	int32 GetHeight() const override
+	{ return Height - Crop.Top - Crop.Bottom; }
+	int32 GetDecodedWidth() const override
+	{ return Width; }
+	int32 GetDecodedHeight() const override
+	{ return Height; }
+	FElectraVideoDecoderOutputCropValues GetCropValues() const override
+	{ return Crop; }
+	int32 GetAspectRatioW() const override
+	{ return AspectW; }
+	int32 GetAspectRatioH() const override
+	{ return AspectH; }
+	int32 GetFrameRateNumerator() const override
+	{ return FrameRateN; }
+	int32 GetFrameRateDenominator() const override
+	{ return FrameRateD; }
+	int32 GetNumberOfBits() const override
+	{ return NumBits; }
+	void GetExtraValues(TMap<FString, FVariant>& OutExtraValues) const override
+	{ OutExtraValues = ExtraValues; }
+	void* GetPlatformOutputHandle(EElectraDecoderPlatformOutputHandleType InTypeOfHandle) const override
+	{ return nullptr; }
+	IElectraDecoderVideoOutputTransferHandle* GetTransferHandle() const override
+	{ return nullptr; }
+	IElectraDecoderVideoOutput::EImageCopyResult CopyPlatformImage(IElectraDecoderVideoOutputCopyResources* InCopyResources) const override;
+
+public:
+	void ReleaseOutputBuffer();
+
+	FTimespan PTS;
+	uint64 UserValue = 0;
+
+	FElectraVideoDecoderOutputCropValues Crop;
+	int32 Width = 0;
+	int32 Height = 0;
+	int32 Pitch = 0;
+	int32 NumBits = 0;
+	int32 AspectW = 1;
+	int32 AspectH = 1;
+	int32 FrameRateN = 0;
+	int32 FrameRateD = 0;
+	int32 PixelFormat = 0;
+	TMap<FString, FVariant> ExtraValues;
+
+	TWeakPtr<FElectraVideoDecoderVPx_Android, ESPMode::ThreadSafe> OwningDecoder;
+	IElectraVPxVideoDecoderAndroidJava::FOutputBufferInfo OwningDecoderBufferInfo;
+	mutable bool bBufferGotReferenced = false;
+};
+
+
+
+class FElectraVideoDecoderVPx_Android : public IElectraVideoDecoderVPx_Android
+{
+public:
+	FElectraVideoDecoderVPx_Android(const TMap<FString, FVariant>& InOptions, TSharedPtr<IElectraDecoderResourceDelegate, ESPMode::ThreadSafe> InResourceDelegate);
+
+	virtual ~FElectraVideoDecoderVPx_Android();
+
+	IElectraDecoder::EType GetType() const override
+	{ return IElectraDecoder::EType::Video; }
+
+	void GetFeatures(TMap<FString, FVariant>& OutFeatures) const override;
+
+	FError GetError() const override;
+
+	void Close() override;
+	IElectraDecoder::ECSDCompatibility IsCompatibleWith(const TMap<FString, FVariant>& CSDAndAdditionalOptions) override;
+	bool ResetToCleanStart() override;
+
+	TSharedPtr<IElectraDecoderDefaultOutputFormat, ESPMode::ThreadSafe> GetDefaultOutputFormatFromCSD(const TMap<FString, FVariant>& CSDAndAdditionalOptions) override;
+
+	EDecoderError DecodeAccessUnit(const FInputAccessUnit& InInputAccessUnit, const TMap<FString, FVariant>& InAdditionalOptions) override;
+	EDecoderError SendEndOfData() override;
+	EDecoderError Flush() override;
+	EOutputStatus HaveOutput() override;
+	TSharedPtr<IElectraDecoderOutput, ESPMode::ThreadSafe> GetOutput() override;
+
+	void Suspend() override;
+	void Resume() override;
+
+	void ReleaseOutputBuffer(const IElectraVPxVideoDecoderAndroidJava::FOutputBufferInfo* InBufferInfo, bool bInRender, int64 InReleaseAt);
+
+private:
+	class FBufferReleaseCallback : public IElectraDecoderResourceDelegateAndroid::IDecoderPlatformResourceAndroid::IBufferReleaseCallback
+	{
+	public:
+		FBufferReleaseCallback(TWeakPtr<FElectraVideoDecoderVPx_Android, ESPMode::ThreadSafe> InOwningDecoder)
+			: OwningDecoder(MoveTemp(InOwningDecoder))
+		{ }
+		virtual ~FBufferReleaseCallback() = default;
+		void OnReleaseBuffer(int32 InBufferIndex, TOptional<int32> InBufferValidCount, TOptional<bool> InDoRender, TOptional<int64> InRenderAt) override
+		{
+			// Note: This could potentially be called from any thread but should be executed only from the decoder thread...
+			check(!"not implemented yet");
+		}
+		TWeakPtr<FElectraVideoDecoderVPx_Android, ESPMode::ThreadSafe> OwningDecoder;
+	};
+
+	class FSurfaceRequestCallback : public IElectraDecoderResourceDelegateAndroid::IDecoderPlatformResourceAndroid::ISurfaceRequestCallback
+	{
+	public:
+		virtual ~FSurfaceRequestCallback() = default;
+		bool HasCompleted() const
+		{ return bIsSet; }
+		ESurfaceType GetType() const
+		{ return SurfaceType; }
+		jobject GetSurface() const
+		{ return Surface; }
+	private:
+		void OnNewSurface(ESurfaceType InSurfaceType, void* InSurface)
+		{
+			SurfaceType = InSurfaceType;
+			Surface = (jobject) InSurface;
+			FPlatformMisc::MemoryBarrier();
+			bIsSet = true;
+		}
+		ESurfaceType SurfaceType = ESurfaceType::Error;
+		jobject Surface = nullptr;
+		volatile bool bIsSet = false;
+	};
+
+	enum class ECodec
+	{
+		VP8,
+		VP9
+	};
+
+	struct FInitialMaxValues
+	{
+		int32 Width = 0;
+		int32 Height = 0;
+		uint32 Profile = 0;
+		uint32 Level = 0;
+		bool IsSet() const
+		{
+			return Width && Height && Profile && Level;
+		}
+	};
+
+	struct FDecoderInput
+	{
+		FInputAccessUnit AccessUnit;
+		TMap<FString, FVariant> AdditionalOptions;
+		ElectraDecodersUtil::VPxVideo::FVP8UncompressedHeader HdrVP8;
+		ElectraDecodersUtil::VPxVideo::FVP9UncompressedHeader HdrVP9;
+	};
+
+	enum class EDecodeState
+	{
+		Decoding,
+		Draining,
+		CreatingDecoder
+	};
+
+	bool PostError(int32 ApiReturnValue, FString Message, int32 Code);
+
+	bool InternalDecoderCreate();
+	bool InternalHandleDecoderCreate();
+	bool InternalFinishDecoderCreate();
+	void InternalDecoderDestroy();
+	bool InternalDecoderDrain();
+
+	enum class EConvertResult
+	{
+		Success,
+		Failure,
+		GotEOS
+	};
+	EConvertResult ConvertDecoderOutput(const IElectraVPxVideoDecoderAndroidJava::FOutputBufferInfo& InInfo);
+	bool HavePendingOutputIndex(int32 InIndex, bool bRemoveIfExists);
+
+	IElectraDecoder::FError LastError;
+	FInitialMaxValues InitialMaxValues;
+	ECodec ActiveCodec = ECodec::VP9;
+
+	TWeakPtr<IElectraDecoderResourceDelegate, ESPMode::ThreadSafe> ResourceDelegate;
+	IElectraDecoderResourceDelegateAndroid::IDecoderPlatformResourceAndroid* PlatformResource = nullptr;
+
+	TSharedPtr<FElectraVideoDecoderOutputVPx_Android, ESPMode::ThreadSafe> CurrentOutput;
+	TArray<TSharedPtr<FDecoderInput, ESPMode::NotThreadSafe>> InDecoderInput;
+
+	TSharedPtr<IElectraVPxVideoDecoderAndroidJava, ESPMode::ThreadSafe> DecoderInstance;
+	IElectraVPxVideoDecoderAndroidJava::FDecoderInformation DecoderInfo;
+	TArray<IElectraVPxVideoDecoderAndroidJava::FOutputBufferInfo> PendingDecoderOutputBuffers;
+	EDecodeState DecodeState = EDecodeState::Decoding;
+	int64 LastPushedPresentationTime = -1;
+	bool bDidSendEOS = false;
+	bool bNewDecoderRequired = false;
+	volatile uint32 DecoderCreatedAtResumeCount = 0;
+	std::atomic<uint32> CurrentResumeCount;
+	bool bTriggerReplay = false;
+
+	TSharedPtr<FBufferReleaseCallback, ESPMode::ThreadSafe> CurrentBufferReleaseCallback;
+	TSharedPtr<FSurfaceRequestCallback, ESPMode::ThreadSafe> CurrentSurfaceRequest;
+	uint32 NativeDecoderID = 0;
+
+	static std::atomic<uint32> NextNativeDecoderID;
+};
+std::atomic<uint32>  FElectraVideoDecoderVPx_Android::NextNativeDecoderID { 0 };
+
+
+namespace IElectraVideoDecoderVPx_Android_Platform
+{
+	static TArray<IElectraVideoDecoderVPx_Android::FSupportedConfiguration> DecoderConfigurations;
+	static bool bDecoderConfigurationsDirty = true;
+}
+
+void IElectraVideoDecoderVPx_Android::PlatformGetSupportedConfigurations(TArray<FSupportedConfiguration>& OutSupportedConfigurations)
+{
+	if (IElectraVideoDecoderVPx_Android_Platform::bDecoderConfigurationsDirty)
+	{
+		IElectraVideoDecoderVPx_Android_Platform::DecoderConfigurations.Empty();
+
+		// Perhaps these can be determined dynamically for each type of device?
+		// What is really supported isn't quite clear so allow for UHD.
+
+		// Main
+		IElectraVideoDecoderVPx_Android_Platform::DecoderConfigurations.Emplace(IElectraVideoDecoderVPx_Android::FSupportedConfiguration(0, 1, 0, 153, 60, 4096, 2304, 0));
+		IElectraVideoDecoderVPx_Android_Platform::DecoderConfigurations.Emplace(IElectraVideoDecoderVPx_Android::FSupportedConfiguration(0, 1, 0, 153, 0, 0, 0, (3840 / 8) * (2160 / 8) * 60));
+		// Main10
+		IElectraVideoDecoderVPx_Android_Platform::DecoderConfigurations.Emplace(IElectraVideoDecoderVPx_Android::FSupportedConfiguration(0, 2, 0, 153, 60, 4096, 2304, 0));
+		IElectraVideoDecoderVPx_Android_Platform::DecoderConfigurations.Emplace(IElectraVideoDecoderVPx_Android::FSupportedConfiguration(0, 2, 0, 153, 0, 0, 0, (3840 / 8) * (2160 / 8) * 60));
+
+		IElectraVideoDecoderVPx_Android_Platform::bDecoderConfigurationsDirty = false;
+	}
+	OutSupportedConfigurations = IElectraVideoDecoderVPx_Android_Platform::DecoderConfigurations;
+}
+
+void IElectraVideoDecoderVPx_Android::GetConfigurationOptions(TMap<FString, FVariant>& OutOptions)
+{
+	OutOptions.Emplace(IElectraDecoderFeature::MinimumNumberOfOutputFrames, FVariant((int32)8));
+	OutOptions.Emplace(IElectraDecoderFeature::IsAdaptive, FVariant(false));
+	OutOptions.Emplace(IElectraDecoderFeature::NeedReplayDataOnDecoderLoss, FVariant(true));
+	OutOptions.Emplace(IElectraDecoderFeature::MustBeSuspendedInBackground, FVariant(true));
+	OutOptions.Emplace(IElectraDecoderFeature::SupportsDroppingOutput, FVariant(true));
+}
+
+TSharedPtr<IElectraDecoder, ESPMode::ThreadSafe> IElectraVideoDecoderVPx_Android::Create(const TMap<FString, FVariant>& InOptions, TSharedPtr<IElectraDecoderResourceDelegate, ESPMode::ThreadSafe> InResourceDelegate)
+{
+	if (InResourceDelegate)
+	{
+		TSharedPtr<FElectraVideoDecoderVPx_Android, ESPMode::ThreadSafe> New = MakeShared<FElectraVideoDecoderVPx_Android>(InOptions, InResourceDelegate);
+		return New;
+	}
+	return nullptr;
+}
+
+
+FElectraVideoDecoderVPx_Android::FElectraVideoDecoderVPx_Android(const TMap<FString, FVariant>& InOptions, TSharedPtr<IElectraDecoderResourceDelegate, ESPMode::ThreadSafe> InResourceDelegate)
+{
+	ResourceDelegate = InResourceDelegate;
+	PlatformResource = reinterpret_cast<IElectraDecoderResourceDelegateAndroid::IDecoderPlatformResourceAndroid*>(ElectraDecodersUtil::GetVariantValueSafeU64(InOptions, TEXT("platformresource"), 0));
+
+	uint32 Codec4CC = (uint32) ElectraDecodersUtil::GetVariantValueSafeI64(InOptions, TEXT("codec_4cc"), 0);
+	ActiveCodec = Codec4CC == 'vp09' ? ECodec::VP9 : ECodec::VP8;
+
+	NativeDecoderID = ++NextNativeDecoderID;
+	DecoderCreatedAtResumeCount = 0;
+	CurrentResumeCount = 0;
+
+	// Get the maximum values if provided.
+	InitialMaxValues.Width = (int32) Align(ElectraDecodersUtil::GetVariantValueSafeI64(InOptions, TEXT("max_width"), 0), 16);
+	InitialMaxValues.Height = (int32) Align(ElectraDecodersUtil::GetVariantValueSafeI64(InOptions, TEXT("max_height"), 0), 16);
+	const FVariant* Var = InOptions.Find(TEXT("max_codecprofile"));
+	if (Var)
+	{
+		ElectraDecodersUtil::FMimeTypeVideoCodecInfo ci;
+		if (ActiveCodec == ECodec::VP9)
+		{
+			if (Var->GetType() == EVariantTypes::String && ElectraDecodersUtil::ParseCodecVP9(ci, Var->GetValue<FString>(), ElectraDecodersUtil::GetVariantValueUInt8Array(InOptions, TEXT("$vpcC_box"))))
+			{
+				InitialMaxValues.Profile = (uint32) ci.Profile;
+				InitialMaxValues.Level = (uint32) ci.Level;
+			}
+		}
+		else
+		{
+			if (Var->GetType() == EVariantTypes::String && ElectraDecodersUtil::ParseCodecVP8(ci, Var->GetValue<FString>(), ElectraDecodersUtil::GetVariantValueUInt8Array(InOptions, TEXT("$vpcC_box"))))
+			{
+				InitialMaxValues.Profile = (uint32) ci.Profile;
+				InitialMaxValues.Level = (uint32) ci.Level;
+			}
+		}
+	}
+}
+
+
+FElectraVideoDecoderVPx_Android::~FElectraVideoDecoderVPx_Android()
+{
+	// Close() must have been called already!
+	check(LastError.Code == ERRCODE_INTERNAL_ALREADY_CLOSED);
+	// We do it nonetheless...
+	Close();
+}
+
+
+void FElectraVideoDecoderVPx_Android::GetFeatures(TMap<FString, FVariant>& OutFeatures) const
+{
+	GetConfigurationOptions(OutFeatures);
+	OutFeatures.Emplace(IElectraDecoderFeature::MustBeSuspendedInBackground, FVariant(true));
+}
+
+
+IElectraDecoder::FError FElectraVideoDecoderVPx_Android::GetError() const
+{
+	return LastError;
+}
+
+
+void FElectraVideoDecoderVPx_Android::Close()
+{
+	ResetToCleanStart();
+	// Set the error state that all subsequent calls will fail.
+	PostError(0, TEXT("Already closed"), ERRCODE_INTERNAL_ALREADY_CLOSED);
+}
+
+
+IElectraDecoder::ECSDCompatibility FElectraVideoDecoderVPx_Android::IsCompatibleWith(const TMap<FString, FVariant>& CSDAndAdditionalOptions)
+{
+	// When we have no decoder yet then we are compatible because we will be creating a decoder when needed.
+	if (!DecoderInstance.IsValid())
+	{
+		return IElectraDecoder::ECSDCompatibility::Compatible;
+	}
+	return DecoderInfo.bIsAdaptive ? IElectraDecoder::ECSDCompatibility::Compatible : IElectraDecoder::ECSDCompatibility::DrainAndReset;
+}
+
+bool FElectraVideoDecoderVPx_Android::ResetToCleanStart()
+{
+	DETAILLOG(LogElectraDecoders, Log, TEXT("VideoDecoderVPx::ResetToCleanStart()"));
+	InternalDecoderDestroy();
+
+	CurrentOutput.Reset();
+	InDecoderInput.Empty();
+	PendingDecoderOutputBuffers.Empty();
+	DecodeState = EDecodeState::Decoding;
+	LastPushedPresentationTime = -1;
+	bDidSendEOS = false;
+	
+	return !LastError.IsSet();
+}
+
+
+TSharedPtr<IElectraDecoderDefaultOutputFormat, ESPMode::ThreadSafe> FElectraVideoDecoderVPx_Android::GetDefaultOutputFormatFromCSD(const TMap<FString, FVariant>& CSDAndAdditionalOptions)
+{
+	return nullptr;
+}
+
+
+IElectraDecoder::EDecoderError FElectraVideoDecoderVPx_Android::DecodeAccessUnit(const FInputAccessUnit& InInputAccessUnit, const TMap<FString, FVariant>& InAdditionalOptions)
+{
+	// If already in error do nothing!
+	if (LastError.IsSet())
+	{
+		return IElectraDecoder::EDecoderError::Error;
+	}
+
+	// Can not feed new input until draining has finished.
+	if (DecodeState == EDecodeState::Draining)
+	{
+		return IElectraDecoder::EDecoderError::EndOfData;
+	}
+
+	// CSD only buffer is not handled at the moment.
+	check((InInputAccessUnit.Flags & EElectraDecoderFlags::InitCSDOnly) == EElectraDecoderFlags::None);
+
+	// If this is discardable and won't be output we do not need to handle it at all.
+	if ((InInputAccessUnit.Flags & (EElectraDecoderFlags::DoNotOutput | EElectraDecoderFlags::IsDiscardable)) == (EElectraDecoderFlags::DoNotOutput | EElectraDecoderFlags::IsDiscardable))
+	{
+		return IElectraDecoder::EDecoderError::None;
+	}
+
+	// Still creating the decoder?
+	if (DecodeState == EDecodeState::CreatingDecoder)
+	{
+		if (!InternalHandleDecoderCreate())
+		{
+			return IElectraDecoder::EDecoderError::Error;
+		}
+		else if (DecodeState == EDecodeState::CreatingDecoder)
+		{
+			return IElectraDecoder::EDecoderError::NoBuffer;
+		}
+	}
+
+	// When using a SurfaceView we need to check if it did not suddenly change.
+	bool bNeedReplayOnNewSurface = false;
+	if (PlatformResource && CurrentSurfaceRequest.IsValid() && CurrentSurfaceRequest->HasCompleted() && CurrentSurfaceRequest->GetType() == FSurfaceRequestCallback::ESurfaceType::SurfaceView)
+	{
+		void* NewSurfaceView = nullptr;
+		IElectraDecoderResourceDelegateAndroid::IDecoderPlatformResourceAndroid::ESurfaceChangeResult SurfaceChangeResult;
+		SurfaceChangeResult = PlatformResource->VerifySurfaceView(NewSurfaceView, reinterpret_cast<void*>(CurrentSurfaceRequest->GetSurface()));
+		if (SurfaceChangeResult == IElectraDecoderResourceDelegateAndroid::IDecoderPlatformResourceAndroid::ESurfaceChangeResult::NewSurface)
+		{
+			// This depends on whether or not the decoder supports setting a new surface. If it doesn't we need to replay on a new instance.
+			check(!"Not currently handled");
+			bNeedReplayOnNewSurface = true;
+		}
+		else if (SurfaceChangeResult == IElectraDecoderResourceDelegateAndroid::IDecoderPlatformResourceAndroid::ESurfaceChangeResult::Error)
+		{
+			PostError(0, TEXT("Failed to get the new output surface"), ERRCODE_INTERNAL_INVALID_OUTPUT_SURFACE);
+			return IElectraDecoder::EDecoderError::Error;
+		}
+	}
+
+	// Resuming after suspending will re-create a decoder by replaying previous data.
+	// This may not be necessary on all devices, but is a safe catch-all.
+	if (DecoderCreatedAtResumeCount != CurrentResumeCount || bNeedReplayOnNewSurface)
+	{
+		bTriggerReplay = true;
+		CurrentOutput.Reset();
+		InDecoderInput.Empty();
+		PendingDecoderOutputBuffers.Empty();
+		LastPushedPresentationTime = -1;
+		bNewDecoderRequired = true;
+	}
+
+	// Create decoder transform if necessary.
+	if (bNewDecoderRequired)
+	{
+		DETAILLOG(LogElectraDecoders, Log, TEXT("VideoDecoderVPx::DecodeAccessUnit() - require new decoder, destroying..."));
+		InternalDecoderDestroy();
+	}
+	if (!DecoderInstance.IsValid())
+	{
+		InternalDecoderCreate();
+		return InternalHandleDecoderCreate() ? IElectraDecoder::EDecoderError::NoBuffer : IElectraDecoder::EDecoderError::Error;
+	}
+
+	if (bTriggerReplay)
+	{
+		bTriggerReplay = false;
+		return IElectraDecoder::EDecoderError::LostDecoder;
+	}
+
+	// Decode the data if given.
+	if (InInputAccessUnit.Data && InInputAccessUnit.DataSize)
+	{
+		ElectraDecodersUtil::VPxVideo::FVP8UncompressedHeader HdrVP8;
+		ElectraDecodersUtil::VPxVideo::FVP9UncompressedHeader HdrVP9;
+		if (ActiveCodec == ECodec::VP9)
+		{
+			if (!ElectraDecodersUtil::VPxVideo::ParseVP9UncompressedHeader(HdrVP9, InInputAccessUnit.Data, InInputAccessUnit.DataSize))
+			{
+				PostError(0, TEXT("Failed to validate VP9 bitstream header"), ERRCODE_INTERNAL_FAILED_TO_DECODE_INPUT);
+				return IElectraDecoder::EDecoderError::Error;
+			}
+		}
+		else //if (ActiveCodec == ECodec::VP8)
+		{
+			if (!ElectraDecodersUtil::VPxVideo::ParseVP8UncompressedHeader(HdrVP8, InInputAccessUnit.Data, InInputAccessUnit.DataSize))
+			{
+				PostError(0, TEXT("Failed to validate VP8 bitstream header"), ERRCODE_INTERNAL_FAILED_TO_DECODE_INPUT);
+				return IElectraDecoder::EDecoderError::Error;
+			}
+		}
+
+		int32 InputBufferIndex = DecoderInstance->DequeueInputBuffer(0);
+		if (InputBufferIndex >= 0)
+		{
+			// Android decoders need the PTS in microseconds
+			int64 PTS = InInputAccessUnit.PTS.GetTicks() / 10;
+			int32 Result = 0;
+			// If this is a sync sample we prepend the CSD. While this is not necessary on a running stream we need to have the CSD
+			// on the first frame and it is easier to prepend it to all IDR frames when seeking etc.
+			if ((InInputAccessUnit.Flags & EElectraDecoderFlags::IsSyncSample) != EElectraDecoderFlags::None)
+			{
+				TArray<uint8> CSD = ElectraDecodersUtil::GetVariantValueUInt8Array(InAdditionalOptions, TEXT("csd"));
+				Result = DecoderInstance->QueueInputBuffer(InputBufferIndex, CSD.GetData(), CSD.Num(), InInputAccessUnit.Data, InInputAccessUnit.DataSize, PTS);
+			}
+			else
+			{
+				Result = DecoderInstance->QueueInputBuffer(InputBufferIndex, nullptr, 0, InInputAccessUnit.Data, InInputAccessUnit.DataSize, PTS);
+			}
+			DETAILLOG(LogElectraDecoders, Log, TEXT("VideoDecoderVPx::DecodeAccessUnit() - queue input %lld, 0x%x"), (long long int)InInputAccessUnit.PTS.GetTicks(), InInputAccessUnit.Flags);
+			if (Result == 0)
+			{
+				LastPushedPresentationTime = PTS;
+
+				// Add to the list of inputs passed to the decoder.
+				TSharedPtr<FDecoderInput, ESPMode::NotThreadSafe> In(new FDecoderInput);
+				In->AdditionalOptions = InAdditionalOptions;
+				In->AccessUnit = InInputAccessUnit;
+				In->HdrVP8 = HdrVP8;
+				In->HdrVP9 = HdrVP9;
+				// Zero the input pointer and size in the copy. That data is not owned by us and it's best not to have any
+				// values here that would lead us to think that we do.
+				In->AccessUnit.Data = nullptr;
+				In->AccessUnit.DataSize = 0;
+				InDecoderInput.Emplace(MoveTemp(In));
+				InDecoderInput.Sort([](const TSharedPtr<FDecoderInput, ESPMode::NotThreadSafe>& e1, const TSharedPtr<FDecoderInput, ESPMode::NotThreadSafe>& e2)
+				{
+					return e1->AccessUnit.PTS < e2->AccessUnit.PTS;
+				});
+
+				return IElectraDecoder::EDecoderError::None;
+			}
+			else
+			{
+				PostError(Result, TEXT("Failed to decode video decoder input"), ERRCODE_INTERNAL_FAILED_TO_DECODE_INPUT);
+				return IElectraDecoder::EDecoderError::Error;
+			}
+		}
+		else if (InputBufferIndex == -1)
+		{
+			// No available input buffer. Try later.
+			return IElectraDecoder::EDecoderError::NoBuffer;
+		}
+		else
+		{
+			PostError(InputBufferIndex, TEXT("Failed to get a decoder input buffer"), ERRCODE_INTERNAL_DID_NOT_GET_INPUT_BUFFER);
+			return IElectraDecoder::EDecoderError::Error;
+		}
+	}
+	return IElectraDecoder::EDecoderError::None;
+}
+
+
+IElectraDecoder::EDecoderError FElectraVideoDecoderVPx_Android::SendEndOfData()
+{
+	DETAILLOG(LogElectraDecoders, Log, TEXT("VideoDecoderVPx::SendEndOfData()"));
+	// If already in error do nothing!
+	if (LastError.IsSet())
+	{
+		return IElectraDecoder::EDecoderError::Error;
+	}
+
+	// Already draining?
+	if (DecodeState == EDecodeState::Draining)
+	{
+		DETAILLOG(LogElectraDecoders, Log, TEXT("VideoDecoderVPx::SendEndOfData() - already sent"));
+		return IElectraDecoder::EDecoderError::EndOfData;
+	}
+	// If there is a transform send an end-of-stream and drain message.
+	if (DecoderInstance.IsValid())
+	{
+		if (!InternalDecoderDrain())
+		{
+			return IElectraDecoder::EDecoderError::Error;
+		}
+		DecodeState = EDecodeState::Draining;
+	}
+	return IElectraDecoder::EDecoderError::None;
+}
+
+
+IElectraDecoder::EDecoderError FElectraVideoDecoderVPx_Android::Flush()
+{
+	DETAILLOG(LogElectraDecoders, Log, TEXT("VideoDecoderVPx::Flush()"));
+	// If already in error do nothing!
+	if (LastError.IsSet())
+	{
+		return IElectraDecoder::EDecoderError::Error;
+	}
+
+#ifdef DESTROY_DECODER_WHEN_FLUSHING
+	ResetToCleanStart();
+
+#else
+	if (DecoderInstance.IsValid())
+	{
+		DecoderInstance->Flush();
+	}
+	CurrentOutput.Reset();
+	InDecoderInput.Empty();
+	PendingDecoderOutputBuffers.Empty();
+	LastPushedPresentationTime = -1;
+	bDidSendEOS = false;
+	DecodeState = EDecodeState::Decoding;
+#endif
+	return IElectraDecoder::EDecoderError::None;
+}
+
+
+IElectraDecoder::EOutputStatus FElectraVideoDecoderVPx_Android::HaveOutput()
+{
+	// If already in error do nothing!
+	if (LastError.IsSet())
+	{
+		return IElectraDecoder::EOutputStatus::Error;
+	}
+
+	if (CurrentOutput.IsValid())
+	{
+		DETAILLOG(LogElectraDecoders, Log, TEXT("VideoDecoderVPx::HaveOutput() -> Have"));
+		return IElectraDecoder::EOutputStatus::Available;
+	}
+
+	if (bNewDecoderRequired || !DecoderInstance.IsValid())
+	{
+		DETAILLOG(LogElectraDecoders, Log, TEXT("VideoDecoderVPx::HaveOutput() no / need new decoder"));
+		return IElectraDecoder::EOutputStatus::NeedInput;
+	}
+
+	// Are we to drain the decoder?
+	// We need an available input buffer to do this so we put it here where pulling output is sure to free
+	// up an input buffer eventually.
+	if (DecodeState == EDecodeState::Draining && !bDidSendEOS)
+	{
+		int32 Result = -1;
+		int32 InputBufferIndex = DecoderInstance->DequeueInputBuffer(0);
+		DETAILLOG(LogElectraDecoders, Log, TEXT("VideoDecoderVPx::HaveOutput() Sending EOS"));
+		if (InputBufferIndex >= 0)
+		{
+			Result = DecoderInstance->QueueEOSInputBuffer(InputBufferIndex, LastPushedPresentationTime);
+			if (Result == 0)
+			{
+				bDidSendEOS = true;
+			}
+			else
+			{
+				PostError(Result, TEXT("Failed to submit decoder EOS input buffer"), ERRCODE_INTERNAL_FAILED_TO_HANDLE_OUTPUT);
+				return IElectraDecoder::EOutputStatus::Error;
+			}
+		}
+		else if (InputBufferIndex == -1)
+		{
+			// No available input buffer. Try later.
+		}
+		else
+		{
+			PostError(InputBufferIndex, TEXT("Failed to get a decoder input buffer for EOS"), ERRCODE_INTERNAL_FAILED_TO_HANDLE_OUTPUT);
+			return IElectraDecoder::EOutputStatus::Error;
+		}
+	}
+
+	// Check if there is available output from the decoder.
+	while(1)
+	{
+		IElectraVPxVideoDecoderAndroidJava::FOutputBufferInfo OutputBufferInfo;
+		int32 Result = DecoderInstance->DequeueOutputBuffer(OutputBufferInfo, 0);
+		if (Result != 0)
+		{
+			PostError(Result, TEXT("Failed to get decoder output buffer"), ERRCODE_INTERNAL_DID_NOT_GET_OUTPUT_BUFFER);
+			return IElectraDecoder::EOutputStatus::Error;
+		}
+
+		// Valid output buffer?
+		if (OutputBufferInfo.BufferIndex >= 0)
+		{
+			// Is this a pure EOS buffer with no decoded data?
+			if (OutputBufferInfo.bIsEOS && OutputBufferInfo.Size == 0)
+			{
+				DETAILLOG(LogElectraDecoders, Log, TEXT("VideoDecoderVPx::HaveOutput() got EOS, %d still in?"), InDecoderInput.Num());
+				// We have to release that buffer even if it's empty.
+				Result = DecoderInstance->ReleaseOutputBuffer(OutputBufferInfo.BufferIndex, OutputBufferInfo.ValidCount, false, -1);
+
+				// Did we ask for an EOS?
+				if (DecodeState == EDecodeState::Draining)
+				{
+					DecodeState = EDecodeState::Decoding;
+					bDidSendEOS = false;
+					InDecoderInput.Empty();
+					// The decoder will not be accepting new input after an EOS until it gets flushed.
+					// Since it will be drained mostly because of a format change we rather create a new one.
+					bNewDecoderRequired = true;
+					return IElectraDecoder::EOutputStatus::EndOfData;
+				}
+				// On some Android versions there could be such empty buffers without us having requested an EOS.
+				else
+				{
+					DETAILLOG(LogElectraDecoders, Log, TEXT("  -> UNEXPECTED, continuing!"));
+					continue;
+				}
+			}
+			// If a configuration buffer is returned to us we ignore it. This serves no purpose.
+			if (OutputBufferInfo.bIsConfig)
+			{
+				Result = DecoderInstance->ReleaseOutputBuffer(OutputBufferInfo.BufferIndex, OutputBufferInfo.ValidCount, false, -1);
+				continue;
+			}
+
+			EConvertResult CnvRes = ConvertDecoderOutput(OutputBufferInfo);
+			if (CnvRes == EConvertResult::Success)
+			{
+				// Was there the EOS flag set and did we request it?
+				if (OutputBufferInfo.bIsEOS && DecodeState == EDecodeState::Draining)
+				{
+					// This should not happen. But assuming that it does then return EOS instead.
+					DecodeState = EDecodeState::Decoding;
+					bDidSendEOS = false;
+					InDecoderInput.Empty();
+					CurrentOutput.Reset();
+					bNewDecoderRequired = true;
+					return IElectraDecoder::EOutputStatus::EndOfData;
+				}
+				return IElectraDecoder::EOutputStatus::Available;
+			}
+			else
+			{
+				return IElectraDecoder::EOutputStatus::Error;
+			}
+		}
+		else if (OutputBufferInfo.BufferIndex == IElectraVPxVideoDecoderAndroidJava::FOutputBufferInfo::EBufferIndexValues::MediaCodec_INFO_TRY_AGAIN_LATER)
+		{
+			if (DecodeState == EDecodeState::Draining)
+			{
+				DETAILLOG(LogElectraDecoders, Log, TEXT("VideoDecoderVPx::HaveOutput() - DRAINING - TRY AGAIN LATER!?!"));
+				return IElectraDecoder::EOutputStatus::TryAgainLater;
+			}
+			DETAILLOG(LogElectraDecoders, Log, TEXT("VideoDecoderVPx::HaveOutput() - TRY AGAIN LATER!?!"));
+			return IElectraDecoder::EOutputStatus::NeedInput;
+		}
+		else if (OutputBufferInfo.BufferIndex == IElectraVPxVideoDecoderAndroidJava::FOutputBufferInfo::EBufferIndexValues::MediaCodec_INFO_OUTPUT_FORMAT_CHANGED)
+		{
+			// Not needed.
+		}
+		else if (OutputBufferInfo.BufferIndex == IElectraVPxVideoDecoderAndroidJava::FOutputBufferInfo::EBufferIndexValues::MediaCodec_INFO_OUTPUT_BUFFERS_CHANGED)
+		{
+			// No-op as this is the Result of a deprecated API we are not using.
+		}
+		else
+		{
+			PostError(Result, TEXT("Unhandled output buffer index value"), ERRCODE_INTERNAL_UNEXPECTED_PROBLEM);
+			return IElectraDecoder::EOutputStatus::Error;
+		}
+	}
+	check(!"should never get here!");
+	return IElectraDecoder::EOutputStatus::Error;
+}
+
+
+TSharedPtr<IElectraDecoderOutput, ESPMode::ThreadSafe> FElectraVideoDecoderVPx_Android::GetOutput()
+{
+	TSharedPtr<IElectraDecoderOutput, ESPMode::ThreadSafe> Out = CurrentOutput;
+	CurrentOutput.Reset();
+	return Out;
+}
+
+
+void FElectraVideoDecoderVPx_Android::Suspend()
+{
+	DETAILLOG(LogElectraDecoders, Log, TEXT("VideoDecoderVPx::Suspend()"));
+}
+
+
+void FElectraVideoDecoderVPx_Android::Resume()
+{
+	DETAILLOG(LogElectraDecoders, Log, TEXT("VideoDecoderVPx::Resume()"));
+	// Increment the resume counter to re-create the decoder.
+	++CurrentResumeCount;
+}
+
+
+bool FElectraVideoDecoderVPx_Android::PostError(int32 ApiReturnValue, FString Message, int32 Code)
+{
+	LastError.Code = Code;
+	LastError.SdkCode = ApiReturnValue;
+	LastError.Message = MoveTemp(Message);
+	return false;
+}
+
+
+bool FElectraVideoDecoderVPx_Android::InternalDecoderCreate()
+{
+	bDidSendEOS = false;
+	bNewDecoderRequired = false;
+
+	CurrentSurfaceRequest.Reset();
+	CurrentBufferReleaseCallback.Reset();
+
+	DecodeState = EDecodeState::CreatingDecoder;
+	return true;
+}
+
+
+bool FElectraVideoDecoderVPx_Android::InternalHandleDecoderCreate()
+{
+	if (CurrentSurfaceRequest.IsValid())
+	{
+		// Request complete?
+		if (!CurrentSurfaceRequest->HasCompleted())
+		{
+			// Not yet. Wait until it has.
+			return true;
+		}
+		switch(CurrentSurfaceRequest->GetType())
+		{
+			//case FSurfaceRequestCallback::ESurfaceType::Error:
+			default:
+			{
+				return PostError(0, TEXT("Failed to create decoder. Failed to obtain a surface to decode onto."), ERRCODE_INTERNAL_COULD_NOT_CREATE_DECODER);
+			}
+			case FSurfaceRequestCallback::ESurfaceType::NoSurface:
+			{
+				return PostError(0, TEXT("Failed to create decoder. Decoding into CPU buffer is not supported."), ERRCODE_INTERNAL_COULD_NOT_CREATE_DECODER);
+			}
+			case FSurfaceRequestCallback::ESurfaceType::Surface:
+			{
+				return InternalFinishDecoderCreate();
+			}
+			case FSurfaceRequestCallback::ESurfaceType::SurfaceView:
+			{
+			// TODO: When we are to use an externally provided surface we need to check with every decode call if it has changed.
+			//       If so we need to update the decoder, or if this fails re-create and replay it.
+				check(!"implement this!");
+				return PostError(0, TEXT("Failed to create decoder. Surface view is not currently supported."), ERRCODE_INTERNAL_COULD_NOT_CREATE_DECODER);
+			}
+		}
+	}
+	else
+	{
+		// We need a surface, for which we need to have a platform resource instance to get it.
+		if (!PlatformResource)
+		{
+			return PostError(0, TEXT("Failed to create decoder. No platform resource handle is set."), ERRCODE_INTERNAL_COULD_NOT_CREATE_DECODER);
+		}
+		// Create a request to get a surface and request one.
+		CurrentSurfaceRequest = MakeShared<FSurfaceRequestCallback, ESPMode::ThreadSafe>();
+		PlatformResource->RequestSurface(CurrentSurfaceRequest);
+		return true;
+	}
+	return true;
+}
+
+bool FElectraVideoDecoderVPx_Android::InternalFinishDecoderCreate()
+{
+	if (CurrentSurfaceRequest.IsValid())
+	{
+		DecoderInstance = IElectraVPxVideoDecoderAndroidJava::Create();
+
+		// Remember at which resume count we created the decoder.
+		DecoderCreatedAtResumeCount = CurrentResumeCount;
+
+		IElectraVPxVideoDecoderAndroidJava::FCreateParameters cp;
+		cp.CodecData = nullptr;
+		cp.MaxWidth = InitialMaxValues.Width ? InitialMaxValues.Width : 640;
+		cp.MaxHeight = InitialMaxValues.Height ? InitialMaxValues.Height : 480;
+		cp.MaxProfile = InitialMaxValues.Profile;
+		cp.MaxFrameRate = 60;
+		cp.NativeDecoderID = NativeDecoderID;
+		cp.VideoCodecSurface = CurrentSurfaceRequest->GetSurface();
+		int32 Result = DecoderInstance->CreateDecoder(ActiveCodec == ECodec::VP9);
+		if (Result)
+		{
+			return PostError(Result, TEXT("Failed to create decoder"), ERRCODE_INTERNAL_COULD_NOT_CREATE_DECODER);
+		}
+		Result = DecoderInstance->InitializeDecoder(cp);
+		if (Result)
+		{
+			return PostError(Result, TEXT("Failed to initialize decoder"), ERRCODE_INTERNAL_COULD_NOT_CREATE_DECODER);
+		}
+
+		const IElectraVPxVideoDecoderAndroidJava::FDecoderInformation* DecInf = DecoderInstance->GetDecoderInformation();
+		if (!DecInf)
+		{
+			return PostError(0, TEXT("Failed to initialize decoder, could not get decoder information"), ERRCODE_INTERNAL_COULD_NOT_CREATE_DECODER);
+		}
+		DecoderInfo = *DecInf;
+
+		Result = DecoderInstance->Start();
+		if (Result)
+		{
+			return PostError(Result, TEXT("Failed to start decoder"), ERRCODE_INTERNAL_COULD_NOT_CREATE_DECODER);
+		}
+		DecodeState = EDecodeState::Decoding;
+
+		// Create a new buffer release interface.
+		if (PlatformResource)
+		{
+			CurrentBufferReleaseCallback = MakeShared<FBufferReleaseCallback, ESPMode::ThreadSafe>(StaticCastWeakPtr<FElectraVideoDecoderVPx_Android>(AsWeak()));
+			PlatformResource->SetBufferReleaseCallback(CurrentBufferReleaseCallback);
+		}
+		return true;
+	}
+	else
+	{
+		return PostError(0, TEXT("Failed to start decoder, surface interface disappeared."), ERRCODE_INTERNAL_COULD_NOT_CREATE_DECODER);
+	}
+}
+
+bool FElectraVideoDecoderVPx_Android::InternalDecoderDrain()
+{
+	bDidSendEOS = false;
+	return true;
+}
+
+
+void FElectraVideoDecoderVPx_Android::InternalDecoderDestroy()
+{
+	DETAILLOG(LogElectraDecoders, Log, TEXT("VideoDecoderVPx::InternalDecoderDestroy()"));
+	if (DecoderInstance.IsValid())
+	{
+		CurrentBufferReleaseCallback.Reset();
+		CurrentSurfaceRequest.Reset();
+		DecoderInstance->Flush();
+		DecoderInstance->Stop();
+		DecoderInstance->ReleaseDecoder();
+		DecoderInstance.Reset();
+	}
+}
+
+
+void FElectraVideoDecoderVPx_Android::ReleaseOutputBuffer(const IElectraVPxVideoDecoderAndroidJava::FOutputBufferInfo* InBufferInfo, bool bInRender, int64 InReleaseAt)
+{
+	if (InBufferInfo && DecoderInstance.IsValid() && HavePendingOutputIndex(InBufferInfo->BufferIndex, true))
+	{
+		int32 Result = DecoderInstance->ReleaseOutputBuffer(InBufferInfo->BufferIndex, InBufferInfo->ValidCount, bInRender, InReleaseAt);
+		if (Result)
+		{
+			PostError(Result, TEXT("Failed to release decoder output buffer"), ERRCODE_INTERNAL_FAILED_TO_CONVERT_OUTPUT_SAMPLE);
+		}
+	}
+}
+
+bool FElectraVideoDecoderVPx_Android::HavePendingOutputIndex(int32 InIndex, bool bRemoveIfExists)
+{
+	for(int32 i=0; i<PendingDecoderOutputBuffers.Num(); ++i)
+	{
+		if (PendingDecoderOutputBuffers[i].BufferIndex == InIndex)
+		{
+			if (bRemoveIfExists)
+			{
+				PendingDecoderOutputBuffers.RemoveAtSwap(i);
+			}
+			return true;
+		}
+	}
+	return false;
+}
+
+
+
+FElectraVideoDecoderVPx_Android::EConvertResult FElectraVideoDecoderVPx_Android::ConvertDecoderOutput(const IElectraVPxVideoDecoderAndroidJava::FOutputBufferInfo& InInfo)
+{
+	DETAILLOG(LogElectraDecoders, Log, TEXT("VideoDecoderVPx::ConvertDecoderOutput()"));
+	if (!InDecoderInput.Num())
+	{
+		PostError(0, TEXT("There is no pending decoder input for the decoded output!"), ERRCODE_INTERNAL_FAILED_TO_CONVERT_OUTPUT_SAMPLE);
+		return EConvertResult::Failure;
+	}
+
+	// Find the input corresponding to this output.
+	TSharedPtr<FDecoderInput, ESPMode::NotThreadSafe> In;
+	for(int32 nInDec=0; nInDec<InDecoderInput.Num(); ++nInDec)
+	{
+		// Decoder timestamps are in microseconds, so we need to compare as that.
+		if (InDecoderInput[nInDec]->AccessUnit.PTS / 10 == InInfo.PresentationTimestamp)
+		{
+			In = InDecoderInput[nInDec];
+			InDecoderInput.RemoveAt(nInDec);
+			break;
+		}
+	}
+	if (!In.IsValid())
+	{
+		PostError(0, TEXT("There is no matching decoder input for the decoded output!"), ERRCODE_INTERNAL_FAILED_TO_CONVERT_OUTPUT_SAMPLE);
+		return EConvertResult::Failure;
+	}
+
+	// If this is a replay frame or a frame that will not be output, we do not need to set it up and merely release the buffer immediately.
+	if ((In->AccessUnit.Flags & (EElectraDecoderFlags::IsReplaySample | EElectraDecoderFlags::DoNotOutput)) != EElectraDecoderFlags::None)
+	{
+		DETAILLOG(LogElectraDecoders, Log, TEXT("VideoDecoderVPx::ConvertDecoderOutput() - Discard replay output %lld, 0x%x"), (long long int)In->AccessUnit.PTS.GetTicks(), In->AccessUnit.Flags);
+		if (DecoderInstance.IsValid())
+		{
+			DecoderInstance->ReleaseOutputBuffer(InInfo.BufferIndex, InInfo.ValidCount, false, -1);
+		}
+		return EConvertResult::Success;
+	}
+	DETAILLOG(LogElectraDecoders, Log, TEXT("VideoDecoderVPx::ConvertDecoderOutput() - Convert output %lld, 0x%x"), (long long int)In->AccessUnit.PTS.GetTicks(), In->AccessUnit.Flags);
+
+	IElectraVPxVideoDecoderAndroidJava::FOutputFormatInfo Info;
+	if (!DecoderInstance.IsValid() || DecoderInstance->GetOutputFormatInfo(Info, InInfo.BufferIndex) != 0)
+	{
+		PostError(0, TEXT("Could not get decoded buffer format info!"), ERRCODE_INTERNAL_FAILED_TO_CONVERT_OUTPUT_SAMPLE);
+		return EConvertResult::Failure;
+	}
+
+	TSharedPtr<FElectraVideoDecoderOutputVPx_Android, ESPMode::ThreadSafe> NewOutput = MakeShared<FElectraVideoDecoderOutputVPx_Android>();
+	NewOutput->PTS = In->AccessUnit.PTS;
+	NewOutput->UserValue = In->AccessUnit.UserValue;
+	NewOutput->Width = Info.Width;
+	NewOutput->Height = Info.Height;
+	NewOutput->Pitch = Info.Stride;
+	NewOutput->NumBits = ActiveCodec == ECodec::VP8 ? 8 : In->HdrVP9.GetBitDepth();
+	NewOutput->PixelFormat = 0;
+	NewOutput->Crop.Left = Info.CropLeft;
+	NewOutput->Crop.Right = Info.Width - Info.CropRight - 1;
+	NewOutput->Crop.Top = Info.CropTop;
+	NewOutput->Crop.Bottom = Info.Height - Info.CropBottom - 1;
+	NewOutput->AspectW = 1;
+	NewOutput->AspectH = 1;
+	NewOutput->OwningDecoder = StaticCastWeakPtr<FElectraVideoDecoderVPx_Android>(AsWeak());
+	NewOutput->OwningDecoderBufferInfo = InInfo;
+
+	NewOutput->ExtraValues.Emplace(TEXT("platform"), FVariant(TEXT("android")));
+	NewOutput->ExtraValues.Emplace(TEXT("codec"), FVariant(ActiveCodec == ECodec::VP9 ? TEXT("vp9") : TEXT("vp8")));
+
+	PendingDecoderOutputBuffers.Emplace(InInfo);
+	CurrentOutput = MoveTemp(NewOutput);
+	return EConvertResult::Success;
+}
+
+
+void FElectraVideoDecoderOutputVPx_Android::ReleaseOutputBuffer()
+{
+	if (OwningDecoderBufferInfo.BufferIndex >= 0)
+	{
+		TSharedPtr<FElectraVideoDecoderVPx_Android, ESPMode::ThreadSafe> Decoder = OwningDecoder.Pin();
+		if (Decoder.IsValid())
+		{
+			Decoder->ReleaseOutputBuffer(&OwningDecoderBufferInfo, bBufferGotReferenced, -1);
+		}
+	}
+}
+
+IElectraDecoderVideoOutput::EImageCopyResult FElectraVideoDecoderOutputVPx_Android::CopyPlatformImage(IElectraDecoderVideoOutputCopyResources* InCopyResources) const
+{ 
+	TSharedPtr<FElectraVideoDecoderVPx_Android, ESPMode::ThreadSafe> Decoder = OwningDecoder.Pin();
+	if (Decoder.IsValid())
+	{
+		bBufferGotReferenced = true;
+
+		InCopyResources->SetBufferIndex(OwningDecoderBufferInfo.BufferIndex);
+		InCopyResources->SetValidCount(OwningDecoderBufferInfo.ValidCount);
+
+		return IElectraDecoderVideoOutput::EImageCopyResult::Succeeded; 
+	}
+	return IElectraDecoderVideoOutput::EImageCopyResult::Failed; 
+}
diff --git a/Engine/Plugins/Media/ElectraCodecs/Source/ElectraDecoders/Private/Android/vpx/VideoDecoderVPx_Android.h b/Engine/Plugins/Media/ElectraCodecs/Source/ElectraDecoders/Private/Android/vpx/VideoDecoderVPx_Android.h
new file mode 100644
--- /dev/null
+++ b/Engine/Plugins/Media/ElectraCodecs/Source/ElectraDecoders/Private/Android/vpx/VideoDecoderVPx_Android.h
@@ -0,0 +1,33 @@
+// Copyright Epic Games, Inc. All Rights Reserved.
+
+#pragma once
+
+
+#include "IElectraDecoder.h"
+#include "IElectraDecoderResourceDelegate.h"
+
+class IElectraVideoDecoderVPx_Android : public IElectraDecoder
+{
+public:
+	struct FSupportedConfiguration
+	{
+		FSupportedConfiguration(int32 InTier, int32 InProfile, uint32 InProfileCompatibility, int32 InLevel, int32 InFramesPerSecond, int32 InWidth, int32 InHeight, int32 InNum8x8Macroblocks)
+			: Tier(InTier), Profile(InProfile), ProfileCompatibility(InProfileCompatibility), Level(InLevel), FramesPerSecond(InFramesPerSecond), Width(InWidth), Height(InHeight), Num8x8Macroblocks(InNum8x8Macroblocks)
+		{}
+
+		int32 Tier = 0;
+		int32 Profile = 0;
+		int32 ProfileCompatibility = 0;
+		int32 Level = 0;
+		int32 FramesPerSecond = 0;
+		int32 Width = 0;
+		int32 Height = 0;
+		int32 Num8x8Macroblocks = 0;
+	};
+	static void PlatformGetSupportedConfigurations(TArray<FSupportedConfiguration>& OutSupportedConfigurations);
+
+	static void GetConfigurationOptions(TMap<FString, FVariant>& OutOptions);
+	static TSharedPtr<IElectraDecoder, ESPMode::ThreadSafe> Create(const TMap<FString, FVariant>& InOptions, TSharedPtr<IElectraDecoderResourceDelegate, ESPMode::ThreadSafe> InResourceDelegate);
+
+	virtual ~IElectraVideoDecoderVPx_Android() = default;
+};
diff --git a/Engine/Plugins/Media/ElectraCodecs/Source/ElectraDecoders/Private/Android/vpx/VideoDecoderVPx_JavaWrapper_Android.cpp b/Engine/Plugins/Media/ElectraCodecs/Source/ElectraDecoders/Private/Android/vpx/VideoDecoderVPx_JavaWrapper_Android.cpp
new file mode 100644
--- /dev/null
+++ b/Engine/Plugins/Media/ElectraCodecs/Source/ElectraDecoders/Private/Android/vpx/VideoDecoderVPx_JavaWrapper_Android.cpp
@@ -0,0 +1,930 @@
+// Copyright Epic Games, Inc. All Rights Reserved.
+
+#include "vpx/VideoDecoderVPx_JavaWrapper_Android.h"
+
+#if PLATFORM_ANDROID
+#include "HAL/Platform.h"
+#include "Utils/Google/ElectraUtilsVPxVideo.h"
+
+#if (defined(USE_ANDROID_JNI_WITHOUT_GAMEACTIVITY) && USE_ANDROID_JNI_WITHOUT_GAMEACTIVITY != 0) || USE_ANDROID_JNI
+#include "Android/AndroidPlatform.h"
+#include "Android/AndroidJava.h"
+
+#if UE_BUILD_SHIPPING
+// always clear any exceptions in shipping
+#define CHECK_JNI_RESULT(Id) if (Id == 0) { JEnv->ExceptionClear(); }
+#else
+#define CHECK_JNI_RESULT(Id)								\
+	if (Id == 0)											\
+	{														\
+		if (bIsOptional)									\
+		{													\
+			JEnv->ExceptionClear();							\
+		}													\
+		else												\
+		{													\
+			JEnv->ExceptionDescribe();						\
+			checkf(Id != 0, TEXT("Failed to find " #Id));	\
+		}													\
+	}
+#endif
+
+//----------------------------------------------------------------------------------------------------
+
+class FElectraVPxVideoDecoderAndroidJava : public IElectraVPxVideoDecoderAndroidJava, public FJavaClassObject
+{
+public:
+
+	FElectraVPxVideoDecoderAndroidJava();
+	virtual ~FElectraVPxVideoDecoderAndroidJava();
+
+	virtual int32 CreateDecoder(bool bCreateVP9) override;
+	virtual int32 InitializeDecoder(const FCreateParameters& InCreateParams) override;
+	virtual int32 SetOutputSurface(jobject InNewOutputSurface) override;
+	virtual int32 ReleaseDecoder() override;
+	virtual const FDecoderInformation* GetDecoderInformation() override;
+	virtual int32 Start() override;
+	virtual int32 Stop() override;
+	virtual int32 Flush() override;
+	virtual int32 Reset() override;
+	virtual int32 DequeueInputBuffer(int32 InTimeoutUsec) override;
+	virtual int32 QueueInputBuffer(int32 InBufferIndex, const void* InPrefixData, int32 InPrefixSize, const void* InAccessUnitData, int32 InAccessUnitSize, int64 InTimestampUSec) override;
+	virtual int32 QueueCSDInputBuffer(int32 InBufferIndex, const void* InCSDData, int32 InCSDSize, int64 InTimestampUSec) override;
+	virtual int32 QueueEOSInputBuffer(int32 InBufferIndex, int64 InTimestampUSec) override;
+	virtual int32 GetOutputFormatInfo(FOutputFormatInfo& OutFormatInfo, int32 InOutputBufferIndex) override;
+	virtual int32 DequeueOutputBuffer(FOutputBufferInfo& OutBufferInfo, int32 InTimeoutUsec) override;
+	virtual int32 GetOutputBuffer(void*& OutBufferDataPtr, int32 OutBufferDataSize, const FOutputBufferInfo& InOutBufferInfo) override;
+	virtual int32 ReleaseOutputBuffer(int32 BufferIndex, int32 ValidCount, bool bRender, int64 releaseAt) override;
+
+private:
+	//-----------------------------------------------------------------------------
+	/**
+	 * Checks for and clears Java exceptions
+	 *
+	 * @param JEnv
+	 */
+	static void GClearException(JNIEnv* JEnv = AndroidJavaEnv::GetJavaEnv())
+	{
+		if (JEnv->ExceptionCheck())
+		{
+			JEnv->ExceptionDescribe();
+			JEnv->ExceptionClear();
+		}
+	}
+
+	static FName GetClassName()
+	{
+		return FName("com/epicgames/unreal/ElectraDecoderVideoVPx");
+	}
+
+	jfieldID FindField(JNIEnv* JEnv, jclass InClass, const ANSICHAR* InFieldName, const ANSICHAR* InFieldType, bool bIsOptional)
+	{
+		jfieldID Field = InClass == nullptr ? nullptr : JEnv->GetFieldID(InClass, InFieldName, InFieldType);
+		CHECK_JNI_RESULT(Field);
+		return Field;
+	}
+
+	jmethodID FindMethod(JNIEnv* JEnv, jclass InClass, const ANSICHAR* InMethodName, const ANSICHAR* InMethodSignature, bool bIsOptional)
+	{
+		jmethodID Method = InClass == nullptr ? nullptr : JEnv->GetMethodID(InClass, InMethodName, InMethodSignature);
+		CHECK_JNI_RESULT(Method);
+		return Method;
+	}
+
+	// Create a Java byte array. Must DeleteLocalRef() after use or handling over to Java.
+	jbyteArray MakeJavaByteArray(const uint8* InData, int32 InNumBytes)
+	{
+		JNIEnv* JEnv = AndroidJavaEnv::GetJavaEnv();
+		jbyteArray RawBuffer = JEnv->NewByteArray(InNumBytes);
+		JEnv->SetByteArrayRegion(RawBuffer, 0, InNumBytes, reinterpret_cast<const jbyte*>(InData));
+		return RawBuffer;
+	}
+
+	template<typename ReturnType>
+	ReturnType CallMethodNoVerify(FJavaClassMethod Method, ...);
+
+	template<>
+	void CallMethodNoVerify<void>(FJavaClassMethod Method, ...)
+	{
+		JNIEnv*	JEnv = AndroidJavaEnv::GetJavaEnv();
+		va_list Params;
+		va_start(Params, Method);
+		JEnv->CallVoidMethodV(Object, Method.Method, Params);
+		va_end(Params);
+	}
+
+	template<>
+	bool CallMethodNoVerify<bool>(FJavaClassMethod Method, ...)
+	{
+		JNIEnv*	JEnv = AndroidJavaEnv::GetJavaEnv();
+		va_list Params;
+		va_start(Params, Method);
+		bool RetVal = JEnv->CallBooleanMethodV(Object, Method.Method, Params);
+		va_end(Params);
+		return RetVal;
+	}
+
+	template<>
+	int CallMethodNoVerify<int>(FJavaClassMethod Method, ...)
+	{
+		JNIEnv*	JEnv = AndroidJavaEnv::GetJavaEnv();
+		va_list Params;
+		va_start(Params, Method);
+		int RetVal = JEnv->CallIntMethodV(Object, Method.Method, Params);
+		va_end(Params);
+		return RetVal;
+	}
+
+	template<>
+	jobject CallMethodNoVerify<jobject>(FJavaClassMethod Method, ...)
+	{
+		JNIEnv*	JEnv = AndroidJavaEnv::GetJavaEnv();
+		va_list Params;
+		va_start(Params, Method);
+		jobject val = JEnv->CallObjectMethodV(Object, Method.Method, Params);
+		va_end(Params);
+		jobject RetVal = JEnv->NewGlobalRef(val);
+		JEnv->DeleteLocalRef(val);
+		return RetVal;
+	}
+
+	template<>
+	jobjectArray CallMethodNoVerify<jobjectArray>(FJavaClassMethod Method, ...)
+	{
+		JNIEnv*	JEnv = AndroidJavaEnv::GetJavaEnv();
+		va_list Params;
+		va_start(Params, Method);
+		jobject val = JEnv->CallObjectMethodV(Object, Method.Method, Params);
+		va_end(Params);
+		jobjectArray RetVal = (jobjectArray)JEnv->NewGlobalRef(val);
+		JEnv->DeleteLocalRef(val);
+		return RetVal;
+	}
+
+	template<>
+	int64 CallMethodNoVerify<int64>(FJavaClassMethod Method, ...)
+	{
+		JNIEnv*	JEnv = AndroidJavaEnv::GetJavaEnv();
+		va_list Params;
+		va_start(Params, Method);
+		int64 RetVal = JEnv->CallLongMethodV(Object, Method.Method, Params);
+		va_end(Params);
+		return RetVal;
+	}
+
+	template<>
+	FString CallMethodNoVerify<FString>(FJavaClassMethod Method, ...)
+	{
+		JNIEnv*	JEnv = AndroidJavaEnv::GetJavaEnv();
+		va_list Params;
+		va_start(Params, Method);
+		jstring RetVal = static_cast<jstring>(
+			JEnv->CallObjectMethodV(Object, Method.Method, Params));
+		va_end(Params);
+		auto Result = FJavaHelper::FStringFromLocalRef(JEnv, RetVal);
+		return Result;
+	}
+
+	void SetupDecoderInformation();
+
+	int32 release();
+
+public:
+	// Java methods
+	FJavaClassMethod	CreateDecoderFN;
+	FJavaClassMethod	ReleaseDecoderFN;
+	FJavaClassMethod	ConfigureDecoderFN;
+	FJavaClassMethod	SetOutputSurfaceFN;
+	FJavaClassMethod	ReleaseFN;
+	FJavaClassMethod	StartFN;
+	FJavaClassMethod	StopFN;
+	FJavaClassMethod	FlushFN;
+	FJavaClassMethod	ResetFN;
+	FJavaClassMethod	DequeueInputBufferFN;
+	FJavaClassMethod	QueueInputBufferFN;
+	FJavaClassMethod	QueueCSDInputBufferFN;
+	FJavaClassMethod	QueueEOSInputBufferFN;
+	FJavaClassMethod	GetDecoderInformationFN;
+	FJavaClassMethod	GetOutputFormatInfoFN;
+	FJavaClassMethod	DequeueOutputBufferFN;
+	FJavaClassMethod	GetOutputBufferFN;
+	FJavaClassMethod	ReleaseOutputBufferFN;
+
+	// FCreateParameters member field IDs
+	jclass				FCreateParametersClass;
+	jmethodID			FCreateParametersCTOR;
+	jfieldID			FCreateParameters_MaxWidth;
+	jfieldID			FCreateParameters_MaxHeight;
+	jfieldID			FCreateParameters_MaxFPS;
+	jfieldID			FCreateParameters_Width;
+	jfieldID			FCreateParameters_Height;
+	jfieldID			FCreateParameters_bNeedSecure;
+	jfieldID			FCreateParameters_bNeedTunneling;
+	jfieldID			FCreateParameters_CSD0;
+	jfieldID			FCreateParameters_NativeDecoderID;
+	jfieldID			FCreateParameters_VideoCodecSurface;
+	jfieldID			FCreateParameters_bSurfaceIsView;
+
+	// FDecoderInformation member field IDs
+	jclass				FDecoderInformationClass;
+	jfieldID			FDecoderInformation_bIsAdaptive;
+	jfieldID			FDecoderInformation_ApiLevel;
+	jfieldID			FDecoderInformation_bCanUse_SetOutputSurface;
+
+	// FOutputFormatInfo member field IDs
+	jclass				FOutputFormatInfoClass;
+	jfieldID			FOutputFormatInfo_Width;
+	jfieldID			FOutputFormatInfo_Height;
+	jfieldID			FOutputFormatInfo_CropTop;
+	jfieldID			FOutputFormatInfo_CropBottom;
+	jfieldID			FOutputFormatInfo_CropLeft;
+	jfieldID			FOutputFormatInfo_CropRight;
+	jfieldID			FOutputFormatInfo_Stride;
+	jfieldID			FOutputFormatInfo_SliceHeight;
+	jfieldID			FOutputFormatInfo_ColorFormat;
+
+	// FOutputBufferInfo member field IDs
+	jclass				FOutputBufferInfoClass;
+	jfieldID			FOutputBufferInfo_BufferIndex;
+	jfieldID			FOutputBufferInfo_PresentationTimestamp;
+	jfieldID			FOutputBufferInfo_Size;
+	jfieldID			FOutputBufferInfo_bIsEOS;
+	jfieldID			FOutputBufferInfo_bIsConfig;
+
+	// Internal state
+	FCriticalSection				Lock;
+	TUniquePtr<FDecoderInformation>	CurrentDecoderInformation;
+	int32							CurrentValidCount = 0;
+	bool							bHaveDecoder = false;
+	bool							bIsStarted = false;
+};
+
+
+//-----------------------------------------------------------------------------
+/**
+ * Create a Java video decoder wrapper.
+ *
+ * @return Decoder wrapper instance.
+ */
+TSharedPtr<IElectraVPxVideoDecoderAndroidJava, ESPMode::ThreadSafe> IElectraVPxVideoDecoderAndroidJava::Create()
+{
+	return MakeShared<FElectraVPxVideoDecoderAndroidJava>();
+}
+
+
+//-----------------------------------------------------------------------------
+/**
+ * CTOR
+ */
+FElectraVPxVideoDecoderAndroidJava::FElectraVPxVideoDecoderAndroidJava()
+	: FJavaClassObject(GetClassName(), "()V")
+	, CreateDecoderFN(GetClassMethod("CreateDecoder", "(Z)I"))
+	, ReleaseDecoderFN(GetClassMethod("ReleaseDecoder", "()I"))
+	, ConfigureDecoderFN(GetClassMethod("ConfigureDecoder", "(Lcom/epicgames/unreal/ElectraDecoderVideoVPx$FCreateParameters;)I"))
+	, SetOutputSurfaceFN(GetClassMethod("SetOutputSurface", "(Landroid/view/Surface;)I"))
+	, ReleaseFN(GetClassMethod("release", "()I"))
+	, StartFN(GetClassMethod("Start", "()I"))
+	, StopFN(GetClassMethod("Stop", "()I"))
+	, FlushFN(GetClassMethod("Flush", "()I"))
+	, ResetFN(GetClassMethod("Reset", "()I"))
+	, DequeueInputBufferFN(GetClassMethod("DequeueInputBuffer", "(I)I"))
+	, QueueInputBufferFN(GetClassMethod("QueueInputBuffer", "(IJ[B)I"))
+	, QueueCSDInputBufferFN(GetClassMethod("QueueCSDInputBuffer", "(IJ[B)I"))
+	, QueueEOSInputBufferFN(GetClassMethod("QueueEOSInputBuffer", "(IJ)I"))
+	, GetDecoderInformationFN(GetClassMethod("GetDecoderInformation", "()Lcom/epicgames/unreal/ElectraDecoderVideoVPx$FDecoderInformation;"))
+	, GetOutputFormatInfoFN(GetClassMethod("GetOutputFormatInfo", "(I)Lcom/epicgames/unreal/ElectraDecoderVideoVPx$FOutputFormatInfo;"))
+	, DequeueOutputBufferFN(GetClassMethod("DequeueOutputBuffer", "(I)Lcom/epicgames/unreal/ElectraDecoderVideoVPx$FOutputBufferInfo;"))
+	, GetOutputBufferFN(GetClassMethod("GetOutputBuffer", "(I)[B"))
+	, ReleaseOutputBufferFN(GetClassMethod("ReleaseOutputBuffer", "(IZJ)I"))
+{
+	JNIEnv* JEnv = AndroidJavaEnv::GetJavaEnv();
+
+	// Get field IDs for FCreateParameters class members
+	jclass localCreateParametersClass = AndroidJavaEnv::FindJavaClass("com/epicgames/unreal/ElectraDecoderVideoVPx$FCreateParameters");
+	FCreateParametersClass = (jclass)JEnv->NewGlobalRef(localCreateParametersClass);
+	JEnv->DeleteLocalRef(localCreateParametersClass);
+	FCreateParametersCTOR = FindMethod(JEnv, FCreateParametersClass, "<init>", "()V", false);
+	FCreateParameters_MaxWidth = FindField(JEnv, FCreateParametersClass, "MaxWidth", "I", false);
+	FCreateParameters_MaxHeight = FindField(JEnv, FCreateParametersClass, "MaxHeight", "I", false);
+	FCreateParameters_MaxFPS = FindField(JEnv, FCreateParametersClass, "MaxFPS", "I", false);
+	FCreateParameters_Width = FindField(JEnv, FCreateParametersClass, "Width", "I", false);
+	FCreateParameters_Height = FindField(JEnv, FCreateParametersClass, "Height", "I", false);
+	FCreateParameters_bNeedSecure = FindField(JEnv, FCreateParametersClass, "bNeedSecure", "Z", false);
+	FCreateParameters_bNeedTunneling = FindField(JEnv, FCreateParametersClass, "bNeedTunneling", "Z", false);
+	FCreateParameters_CSD0 = FindField(JEnv, FCreateParametersClass, "CSD0", "[B", false);
+	FCreateParameters_NativeDecoderID = FindField(JEnv, FCreateParametersClass, "NativeDecoderID", "I", false);
+	FCreateParameters_VideoCodecSurface = FindField(JEnv, FCreateParametersClass, "VideoCodecSurface", "Landroid/view/Surface;", false);
+	FCreateParameters_bSurfaceIsView = FindField(JEnv, FCreateParametersClass, "bSurfaceIsView", "Z", false);
+
+	// Get field IDs for FDecoderInformation class members
+	jclass localDecoderInformationClass = AndroidJavaEnv::FindJavaClass("com/epicgames/unreal/ElectraDecoderVideoVPx$FDecoderInformation");
+	FDecoderInformationClass = (jclass)JEnv->NewGlobalRef(localDecoderInformationClass);
+	JEnv->DeleteLocalRef(localDecoderInformationClass);
+	FDecoderInformation_bIsAdaptive = FindField(JEnv, FDecoderInformationClass, "bIsAdaptive", "Z", false);
+	FDecoderInformation_ApiLevel = FindField(JEnv, FDecoderInformationClass, "ApiLevel", "I", false);
+	FDecoderInformation_bCanUse_SetOutputSurface = FindField(JEnv, FDecoderInformationClass, "bCanUse_SetOutputSurface", "Z", false);
+
+	// Get field IDs for FOutputFormatInfo class members
+	jclass localOutputFormatInfoClass = AndroidJavaEnv::FindJavaClass("com/epicgames/unreal/ElectraDecoderVideoVPx$FOutputFormatInfo");
+	FOutputFormatInfoClass = (jclass)JEnv->NewGlobalRef(localOutputFormatInfoClass);
+	JEnv->DeleteLocalRef(localOutputFormatInfoClass);
+	FOutputFormatInfo_Width = FindField(JEnv, FOutputFormatInfoClass, "Width", "I", false);
+	FOutputFormatInfo_Height = FindField(JEnv, FOutputFormatInfoClass, "Height", "I", false);
+	FOutputFormatInfo_CropTop = FindField(JEnv, FOutputFormatInfoClass, "CropTop", "I", false);
+	FOutputFormatInfo_CropBottom = FindField(JEnv, FOutputFormatInfoClass, "CropBottom", "I", false);
+	FOutputFormatInfo_CropLeft = FindField(JEnv, FOutputFormatInfoClass, "CropLeft", "I", false);
+	FOutputFormatInfo_CropRight = FindField(JEnv, FOutputFormatInfoClass, "CropRight", "I", false);
+	FOutputFormatInfo_Stride = FindField(JEnv, FOutputFormatInfoClass, "Stride", "I", false);
+	FOutputFormatInfo_SliceHeight = FindField(JEnv, FOutputFormatInfoClass, "SliceHeight", "I", false);
+	FOutputFormatInfo_ColorFormat = FindField(JEnv, FOutputFormatInfoClass, "ColorFormat", "I", false);
+
+	// Get field IDs for FOutputBufferInfo class members
+	jclass localOutputBufferInfoClass = AndroidJavaEnv::FindJavaClass("com/epicgames/unreal/ElectraDecoderVideoVPx$FOutputBufferInfo");
+	FOutputBufferInfoClass = (jclass)JEnv->NewGlobalRef(localOutputBufferInfoClass);
+	JEnv->DeleteLocalRef(localOutputBufferInfoClass);
+	FOutputBufferInfo_BufferIndex = FindField(JEnv, FOutputBufferInfoClass, "BufferIndex", "I", false);
+	FOutputBufferInfo_PresentationTimestamp = FindField(JEnv, FOutputBufferInfoClass, "PresentationTimestamp", "J", false);
+	FOutputBufferInfo_Size = FindField(JEnv, FOutputBufferInfoClass, "Size", "I", false);
+	FOutputBufferInfo_bIsEOS = FindField(JEnv, FOutputBufferInfoClass, "bIsEOS", "Z", false);
+	FOutputBufferInfo_bIsConfig = FindField(JEnv, FOutputBufferInfoClass, "bIsConfig", "Z", false);
+
+	// Set up the initial decoder information.
+	SetupDecoderInformation();
+}
+
+
+//-----------------------------------------------------------------------------
+/**
+ * DTOR
+ */
+FElectraVPxVideoDecoderAndroidJava::~FElectraVPxVideoDecoderAndroidJava()
+{
+	FScopeLock lock(&Lock);
+	release();
+
+	JNIEnv* JEnv = AndroidJavaEnv::GetJavaEnv();
+	JEnv->DeleteGlobalRef(FCreateParametersClass);
+	JEnv->DeleteGlobalRef(FDecoderInformationClass);
+	JEnv->DeleteGlobalRef(FOutputBufferInfoClass);
+	JEnv->DeleteGlobalRef(FOutputFormatInfoClass);
+
+	CurrentDecoderInformation.Reset();
+}
+
+
+//-----------------------------------------------------------------------------
+/**
+ * Sets up the current decoder information.
+ */
+void FElectraVPxVideoDecoderAndroidJava::SetupDecoderInformation()
+{
+	// Create an instance of the init param structure and fill in the members.
+	JNIEnv* JEnv = AndroidJavaEnv::GetJavaEnv();
+	jobject OutputInfo = JEnv->CallObjectMethod(Object, GetDecoderInformationFN.Method);
+	GClearException(JEnv);
+	// Failure will return no object.
+	check(OutputInfo != nullptr);
+	if (OutputInfo != nullptr)
+	{
+		CurrentDecoderInformation = MakeUnique<FDecoderInformation>();
+		CurrentDecoderInformation->ApiLevel = JEnv->GetIntField(OutputInfo, FDecoderInformation_ApiLevel);
+		CurrentDecoderInformation->bIsAdaptive = JEnv->GetBooleanField(OutputInfo, FDecoderInformation_bIsAdaptive);
+		CurrentDecoderInformation->bCanUse_SetOutputSurface = JEnv->GetBooleanField(OutputInfo, FDecoderInformation_bCanUse_SetOutputSurface);
+		JEnv->DeleteLocalRef(OutputInfo);
+	}
+}
+
+
+//-----------------------------------------------------------------------------
+/**
+ * Creates a Java instance of a VP8 or VP9 video decoder.
+ *
+ * @return 0 if successful, 1 on error.
+ */
+int32 FElectraVPxVideoDecoderAndroidJava::CreateDecoder(bool bCreateVP9)
+{
+	FScopeLock lock(&Lock);
+	int32 result = -1;
+
+	check(!bHaveDecoder);
+
+	CurrentDecoderInformation.Reset();
+
+	// Create and initialize a decoder instance.
+	result = CallMethodNoVerify<int>(CreateDecoderFN, bCreateVP9);
+	GClearException();
+	if (result != 0)
+	{
+		return 1;
+	}
+
+	// Get decoder information
+	SetupDecoderInformation();
+
+	bHaveDecoder = true;
+	return 0;
+}
+
+//-----------------------------------------------------------------------------
+/**
+ * Initializes the video decoder instance.
+ *
+ * @param InCreateParams
+ *
+ * @return 0 if successful, 1 on error.
+ */
+int32 FElectraVPxVideoDecoderAndroidJava::InitializeDecoder(const FCreateParameters& InCreateParams)
+{
+	FScopeLock lock(&Lock);
+	if (bHaveDecoder)
+	{
+		int32 result = -1;
+
+		// Create an instance of the init param structure and fill in the members.
+		JNIEnv* JEnv = AndroidJavaEnv::GetJavaEnv();
+		jobject CreateParams = JEnv->NewObject(FCreateParametersClass, FCreateParametersCTOR);
+		JEnv->SetIntField(CreateParams, FCreateParameters_MaxWidth, InCreateParams.MaxWidth);
+		JEnv->SetIntField(CreateParams, FCreateParameters_MaxHeight, InCreateParams.MaxHeight);
+		JEnv->SetIntField(CreateParams, FCreateParameters_MaxFPS, InCreateParams.MaxFrameRate);
+		//JEnv->SetBooleanField(CreateParams, FCreateParameters_bNeedSecure, false);
+		//JEnv->SetBooleanField(CreateParams, FCreateParameters_bNeedTunneling, false);
+		JEnv->SetIntField(CreateParams, FCreateParameters_NativeDecoderID, InCreateParams.NativeDecoderID);
+/*
+		if (InCreateParams.CodecData)
+		{
+			jbyteArray CSD0 = MakeJavaByteArray(InCreateParams.CodecData->GetCodecSpecificData().GetData(), InCreateParams.CodecData->GetCodecSpecificData().Num());
+			JEnv->SetObjectField(CreateParams, FCreateParameters_CSD0, CSD0);
+			JEnv->DeleteLocalRef(CSD0);
+		}
+*/
+		// Pass along decoder output surface. This is a global ref that we take control over now.
+		jobject GlobalSurfaceRef = InCreateParams.VideoCodecSurface;
+		if (GlobalSurfaceRef)
+		{
+			JEnv->SetObjectField(CreateParams, FCreateParameters_VideoCodecSurface, GlobalSurfaceRef);
+		}
+		JEnv->SetBooleanField(CreateParams, FCreateParameters_bSurfaceIsView, InCreateParams.bSurfaceIsView);
+
+		// Create and initialize a decoder instance.
+		result = CallMethodNoVerify<int>(ConfigureDecoderFN, CreateParams);
+		JEnv->DeleteLocalRef(CreateParams);
+		// We now release the global ref of the surface we were given. The decoder should have taken ownership of it now.
+		if (GlobalSurfaceRef)
+		{
+			JEnv->DeleteGlobalRef(GlobalSurfaceRef);
+		}
+
+		GClearException(JEnv);
+		if (result != 0)
+		{
+			return 1;
+		}
+
+		return 0;
+	}
+	else
+	{
+		return 1;
+	}
+}
+
+
+//-----------------------------------------------------------------------------
+/**
+ * Attempts to set a new output surface on an existing and configured decoder.
+ *
+ * @param InNewOutputSurface
+ *
+ * @return 0 if successful, 1 on error.
+ */
+int32 FElectraVPxVideoDecoderAndroidJava::SetOutputSurface(jobject InNewOutputSurface)
+{
+	FScopeLock lock(&Lock);
+	if (bHaveDecoder)
+	{
+		int32 result = -1;
+		result = CallMethodNoVerify<int>(SetOutputSurfaceFN, InNewOutputSurface);
+		GClearException();
+		if (result != 0)
+		{
+			return 1;
+		}
+		return 0;
+	}
+	else
+	{
+		return 1;
+	}
+}
+
+
+//-----------------------------------------------------------------------------
+/**
+ * Releases (destroys) the Java video decoder instance.
+ *
+ * @return 0 if successful, 1 on error.
+ */
+int32 FElectraVPxVideoDecoderAndroidJava::ReleaseDecoder()
+{
+	FScopeLock lock(&Lock);
+	if (bHaveDecoder)
+	{
+		int32 result = CallMethodNoVerify<int>(ReleaseDecoderFN);
+		GClearException();
+		bHaveDecoder = false;
+		return result ? 1 : 0;
+	}
+	return 1;
+}
+
+
+//-----------------------------------------------------------------------------
+/**
+ * Releases (destroys) the Java texture renderer the Java decoder has created
+ *
+ * @return 0 if successful, 1 on error.
+ */
+int32 FElectraVPxVideoDecoderAndroidJava::release()
+{
+	Stop();
+
+	int32 result = CallMethodNoVerify<int>(ReleaseFN);
+	GClearException();
+	return result ? 1 : 0;
+}
+
+
+//-----------------------------------------------------------------------------
+/**
+ * Returns decoder information after a successful InitializeDecoder().
+ *
+ * @return Pointer to decoder information or null when no decoder has been created.
+ */
+const IElectraVPxVideoDecoderAndroidJava::FDecoderInformation* FElectraVPxVideoDecoderAndroidJava::GetDecoderInformation()
+{
+	return CurrentDecoderInformation.Get();
+}
+
+
+//-----------------------------------------------------------------------------
+/**
+ * Starts the decoder instance.
+ *
+ * @return 0 if successful, 1 on error.
+ */
+int32 FElectraVPxVideoDecoderAndroidJava::Start()
+{
+	FScopeLock lock(&Lock);
+	if (bHaveDecoder && !bIsStarted)
+	{
+		int32 result = CallMethodNoVerify<int>(StartFN);
+		GClearException();
+		if (result)
+		{
+			return result;
+		}
+		bIsStarted = true;
+		return 0;
+	}
+	return 1;
+}
+
+
+//-----------------------------------------------------------------------------
+/**
+ * Stops the decoder instance.
+ *
+ * @return 0 if successful, 1 on error.
+ */
+int32 FElectraVPxVideoDecoderAndroidJava::Stop()
+{
+	FScopeLock lock(&Lock);
+	if (bHaveDecoder && bIsStarted)
+	{
+		int32 result = CallMethodNoVerify<int>(StopFN);
+		GClearException();
+		if (result)
+		{
+			return result;
+		}
+		bIsStarted = false;
+		return 0;
+	}
+	return 1;
+}
+
+
+//-----------------------------------------------------------------------------
+/**
+ * Flushes the decoder instance.
+ *
+ * @return 0 if successful, 1 on error.
+ */
+int32 FElectraVPxVideoDecoderAndroidJava::Flush()
+{
+	FScopeLock lock(&Lock);
+
+	// Need to increment the valid count to ensure we will not try to release or render output buffers
+	// that have become invalid.
+	FPlatformAtomics::InterlockedIncrement(&CurrentValidCount);
+
+	// Synchronously operating decoders must be in the started state to be flushed.
+	if (bHaveDecoder && bIsStarted)
+	{
+		int32 result = CallMethodNoVerify<int>(FlushFN);
+		GClearException();
+		return result;
+	}
+	return 1;
+}
+
+
+//-----------------------------------------------------------------------------
+/**
+ * Resets the decoder instance.
+ *
+ * @return 0 if successful, 1 on error.
+ */
+int32 FElectraVPxVideoDecoderAndroidJava::Reset()
+{
+	FScopeLock lock(&Lock);
+	// Synchronously operating decoders should (must?) be in the stopped state to be reset.
+	if (bHaveDecoder && !bIsStarted)
+	{
+		int32 result = CallMethodNoVerify<int>(ResetFN);
+		GClearException();
+		return result ? 1 : 0;
+	}
+	return 1;
+}
+
+
+//-----------------------------------------------------------------------------
+/**
+ * Dequeues an input buffer.
+ *
+ * @param InTimeoutUsec Timeout in microseconds to wait for an available buffer.
+ *
+ * @return >= 0 returns the index of the successfully dequeued buffer, negative values indicate an error.
+ */
+int32 FElectraVPxVideoDecoderAndroidJava::DequeueInputBuffer(int32 InTimeoutUsec)
+{
+	FScopeLock lock(&Lock);
+	if (bHaveDecoder)
+	{
+		int32 result = CallMethodNoVerify<int>(DequeueInputBufferFN, InTimeoutUsec);
+		GClearException();
+		return result;
+	}
+	return -1;
+}
+
+
+//-----------------------------------------------------------------------------
+/**
+ * Queues input for decoding in the buffer with a previously dequeued (calling DequeueInputBuffer()) index.
+ *
+ * @param InBufferIndex Index of the buffer to put data into and enqueue for decoding (see DequeueInputBuffer()).
+ * @param InPrefixData Data to prefix the data to decode with. If given it is usually the SPS & PPS.
+ * @param InPrefixSize Size of the prefix data.
+ * @param InAccessUnitData Data to be decoded.
+ * @param InAccessUnitSize Size of the data to be decoded.
+ * @param InTimestampUSec Timestamp (PTS) of the data, in microseconds.
+ *
+ * @return 0 if successful, 1 on error.
+ */
+int32 FElectraVPxVideoDecoderAndroidJava::QueueInputBuffer(int32 InBufferIndex, const void* InPrefixData, int32 InPrefixSize, const void* InAccessUnitData, int32 InAccessUnitSize, int64 InTimestampUSec)
+{
+	FScopeLock lock(&Lock);
+	if (bHaveDecoder)
+	{
+		int32 result;
+		JNIEnv* JEnv = AndroidJavaEnv::GetJavaEnv();
+		if (InPrefixData && InPrefixSize)
+		{
+			uint8* Temp = (uint8*)FMemory::Malloc(InPrefixSize + InAccessUnitSize);
+			FMemory::Memcpy(Temp, InPrefixData, InPrefixSize);
+			FMemory::Memcpy(Temp+InPrefixSize, InAccessUnitData, InAccessUnitSize);
+			jbyteArray InData = MakeJavaByteArray((const uint8*)Temp, InPrefixSize + InAccessUnitSize);
+			result = CallMethodNoVerify<int>(QueueInputBufferFN, InBufferIndex, (jlong)InTimestampUSec, InData);
+			JEnv->DeleteLocalRef(InData);
+			FMemory::Free(Temp);
+		}
+		else
+		{
+			jbyteArray InData = MakeJavaByteArray((const uint8*)InAccessUnitData, InAccessUnitSize);
+			result = CallMethodNoVerify<int>(QueueInputBufferFN, InBufferIndex, (jlong)InTimestampUSec, InData);
+			JEnv->DeleteLocalRef(InData);
+		}
+		GClearException(JEnv);
+		return result ? 1 : 0;
+	}
+	return 1;
+}
+
+
+//-----------------------------------------------------------------------------
+/**
+ * Queues codec specific data for the following to-be-decoded data buffers.
+ *
+ * @param InBufferIndex Index of the buffer to put data into and enqueue for decoding (see DequeueInputBuffer()).
+ * @param InCSDData Codec specific data
+ * @param InCSDSize Size of the codec specific data
+ * @param InTimestampUSec Timestamp (PTS) of the data, in microseconds. Must be the same as the next data to be decoded.
+ *
+ * @return 0 if successful, 1 on error.
+ */
+int32 FElectraVPxVideoDecoderAndroidJava::QueueCSDInputBuffer(int32 InBufferIndex, const void* InCSDData, int32 InCSDSize, int64 InTimestampUSec)
+{
+	FScopeLock lock(&Lock);
+	if (bHaveDecoder)
+	{
+		JNIEnv* JEnv = AndroidJavaEnv::GetJavaEnv();
+		jbyteArray  InData = MakeJavaByteArray((const uint8*)InCSDData, InCSDSize);
+		int32 result = CallMethodNoVerify<int>(QueueCSDInputBufferFN, InBufferIndex, (jlong)InTimestampUSec, InData);
+		JEnv->DeleteLocalRef(InData);
+		GClearException(JEnv);
+		return result ? 1 : 0;
+	}
+	return 1;
+}
+
+
+//-----------------------------------------------------------------------------
+/**
+ * Queues end of stream for the buffer with a previously dequeued (calling DequeueInputBuffer()) index.
+ *
+ * @param InBufferIndex Index of the buffer to put the EOS flag into and enqueue for decoding (see DequeueInputBuffer()).
+ * @param InTimestampUSec Timestamp the previous data had. Can be 0.
+ *
+ * @return 0 if successful, 1 on error.
+ */
+int32 FElectraVPxVideoDecoderAndroidJava::QueueEOSInputBuffer(int32 InBufferIndex, int64 InTimestampUSec)
+{
+	FScopeLock lock(&Lock);
+	if (bHaveDecoder)
+	{
+		int32 result = CallMethodNoVerify<int>(QueueEOSInputBufferFN, InBufferIndex, (jlong)InTimestampUSec);
+		GClearException();
+		return result ? 1 : 0;
+	}
+	return 1;
+}
+
+
+//-----------------------------------------------------------------------------
+/**
+ * Returns format information of the decoded samples.
+ *
+ * @param OutFormatInfo
+ * @param InOutputBufferIndex RESERVED FOR NOW - Pass any negative value to get the output format after DequeueOutputBuffer() returns a BufferIndex of MediaCodec_INFO_OUTPUT_FORMAT_CHANGED.
+ *
+ * @return 0 if successful, 1 on error.
+ */
+int32 FElectraVPxVideoDecoderAndroidJava::GetOutputFormatInfo(FOutputFormatInfo& OutFormatInfo, int32 InOutputBufferIndex)
+{
+	FScopeLock lock(&Lock);
+	if (bHaveDecoder)
+	{
+		JNIEnv* JEnv = AndroidJavaEnv::GetJavaEnv();
+		jobject OutputInfo = JEnv->CallObjectMethod(Object, GetOutputFormatInfoFN.Method, InOutputBufferIndex);
+		GClearException(JEnv);
+		// Failure will return no object.
+		if (OutputInfo != nullptr)
+		{
+			OutFormatInfo.Width = JEnv->GetIntField(OutputInfo, FOutputFormatInfo_Width);
+			OutFormatInfo.Height = JEnv->GetIntField(OutputInfo, FOutputFormatInfo_Height);
+			OutFormatInfo.CropTop = JEnv->GetIntField(OutputInfo, FOutputFormatInfo_CropTop);
+			OutFormatInfo.CropBottom = JEnv->GetIntField(OutputInfo, FOutputFormatInfo_CropBottom);
+			OutFormatInfo.CropLeft = JEnv->GetIntField(OutputInfo, FOutputFormatInfo_CropLeft);
+			OutFormatInfo.CropRight = JEnv->GetIntField(OutputInfo, FOutputFormatInfo_CropRight);
+			OutFormatInfo.Stride = JEnv->GetIntField(OutputInfo, FOutputFormatInfo_Stride);
+			OutFormatInfo.SliceHeight = JEnv->GetIntField(OutputInfo, FOutputFormatInfo_SliceHeight);
+			OutFormatInfo.ColorFormat = JEnv->GetIntField(OutputInfo, FOutputFormatInfo_ColorFormat);
+			JEnv->DeleteLocalRef(OutputInfo);
+			return 0;
+		}
+	}
+	return 1;
+}
+
+
+//-----------------------------------------------------------------------------
+/**
+ * Dequeues an output buffer.
+ *
+ * @param InTimeoutUsec Timeout in microseconds to wait for an available buffer.
+ *
+ * @return 0 on success, 1 on failure. The OutBufferInfo.BufferIndex indicates the buffer index.
+ */
+int32 FElectraVPxVideoDecoderAndroidJava::DequeueOutputBuffer(FOutputBufferInfo& OutBufferInfo, int32 InTimeoutUsec)
+{
+	FScopeLock lock(&Lock);
+	if (bHaveDecoder)
+	{
+		JNIEnv* JEnv = AndroidJavaEnv::GetJavaEnv();
+		jobject OutputInfo = JEnv->CallObjectMethod(Object, DequeueOutputBufferFN.Method, InTimeoutUsec);
+		GClearException(JEnv);
+		// Failure will return no object.
+		if (OutputInfo != nullptr)
+		{
+			OutBufferInfo.ValidCount = CurrentValidCount;
+			OutBufferInfo.BufferIndex = JEnv->GetIntField(OutputInfo, FOutputBufferInfo_BufferIndex);
+			OutBufferInfo.PresentationTimestamp = JEnv->GetLongField(OutputInfo, FOutputBufferInfo_PresentationTimestamp);
+			OutBufferInfo.Size = JEnv->GetIntField(OutputInfo, FOutputBufferInfo_Size);
+			OutBufferInfo.bIsEOS = JEnv->GetBooleanField(OutputInfo, FOutputBufferInfo_bIsEOS);
+			OutBufferInfo.bIsConfig = JEnv->GetBooleanField(OutputInfo, FOutputBufferInfo_bIsConfig);
+			JEnv->DeleteLocalRef(OutputInfo);
+			return 0;
+		}
+	}
+	return 1;
+}
+
+
+//-----------------------------------------------------------------------------
+/**
+ * Returns the decoded samples from a decoder output buffer in the decoder native format!
+ *
+ * @param OutBufferDataPtr
+ * @param OutBufferDataSize
+ * @param InOutBufferInfo
+ *
+ * @return 0 on success, 1 on failure.
+ */
+int32 FElectraVPxVideoDecoderAndroidJava::GetOutputBuffer(void*& OutBufferDataPtr, int32 OutBufferDataSize, const FOutputBufferInfo& InOutBufferInfo)
+{
+	FScopeLock lock(&Lock);
+	if (bHaveDecoder)
+	{
+		if (InOutBufferInfo.ValidCount == CurrentValidCount)
+		{
+			JNIEnv* JEnv = AndroidJavaEnv::GetJavaEnv();
+			jbyteArray RawDecoderArray = (jbyteArray)JEnv->CallObjectMethod(Object, GetOutputBufferFN.Method, InOutBufferInfo.BufferIndex);
+			GClearException(JEnv);
+
+			if (RawDecoderArray != nullptr && OutBufferDataPtr != nullptr)
+			{
+				jbyte* RawDataPtr = JEnv->GetByteArrayElements(RawDecoderArray, 0);
+				int32 RawBufferSize = JEnv->GetArrayLength(RawDecoderArray);
+
+				check(RawBufferSize == InOutBufferInfo.Size);
+				check(RawBufferSize <= OutBufferDataSize)
+				if (RawDataPtr != nullptr)
+				{
+					FMemory::Memcpy(OutBufferDataPtr, RawDataPtr, RawBufferSize <= OutBufferDataSize ? RawBufferSize : OutBufferDataSize);
+					JEnv->ReleaseByteArrayElements(RawDecoderArray, RawDataPtr, JNI_ABORT);
+				}
+
+				JEnv->DeleteLocalRef(RawDecoderArray);
+			}
+			return 0;
+		}
+		else
+		{
+			// Decoder got flushed, buffer index no longer valid.
+			return 1;
+		}
+	}
+	return 1;
+}
+
+
+//-----------------------------------------------------------------------------
+/**
+ * Releases the decoder output buffer back to the decoder.
+ *
+ * @param InOutBufferInfo
+ * @param bRender
+ * @param releaseAt
+ *
+ * @return 0 on success, 1 on failure.
+ */
+int32 FElectraVPxVideoDecoderAndroidJava::ReleaseOutputBuffer(int32 BufferIndex, int32 ValidCount, bool bRender, int64 releaseAt)
+{
+	FScopeLock lock(&Lock);
+	if (bHaveDecoder)
+	{
+		// Still same decoder instance?
+		if (ValidCount == CurrentValidCount)
+		{
+			// Yes...
+			int32 result = CallMethodNoVerify<int>(ReleaseOutputBufferFN, BufferIndex, bRender, (long)releaseAt);
+			GClearException();
+			return result ? 1 : 0;
+		}
+		else
+		{
+			// Decoder got flushed, buffer index no longer valid.
+			// This is NOT an error!
+			return 0;
+		}
+	}
+	return 1;
+
+}
+
+
+#endif
+
+#endif
diff --git a/Engine/Plugins/Media/ElectraCodecs/Source/ElectraDecoders/Private/Android/vpx/VideoDecoderVPx_JavaWrapper_Android.h b/Engine/Plugins/Media/ElectraCodecs/Source/ElectraDecoders/Private/Android/vpx/VideoDecoderVPx_JavaWrapper_Android.h
new file mode 100644
--- /dev/null
+++ b/Engine/Plugins/Media/ElectraCodecs/Source/ElectraDecoders/Private/Android/vpx/VideoDecoderVPx_JavaWrapper_Android.h
@@ -0,0 +1,262 @@
+// Copyright Epic Games, Inc. All Rights Reserved.
+
+#pragma once
+
+#include "CoreMinimal.h"
+
+#include "HAL/Platform.h"
+
+#if (defined(USE_ANDROID_JNI_WITHOUT_GAMEACTIVITY) && USE_ANDROID_JNI_WITHOUT_GAMEACTIVITY != 0) || USE_ANDROID_JNI
+#include "Android/AndroidPlatform.h"
+#include "Android/AndroidJava.h"
+
+#include "Utils/Google/ElectraUtilsVPxVideo.h"
+
+class IElectraVPxVideoDecoderAndroidJava
+{
+public:
+
+	static TSharedPtr<IElectraVPxVideoDecoderAndroidJava, ESPMode::ThreadSafe> Create();
+
+	virtual ~IElectraVPxVideoDecoderAndroidJava() = default;
+
+	struct FCreateParameters
+	{
+		int32 MaxWidth = 0;
+		int32 MaxHeight = 0;
+		int32 MaxProfile = 0;
+		int32 MaxFrameRate= 0;
+		const ElectraDecodersUtil::VPxVideo::FVP9UncompressedHeader* CodecData = nullptr;
+		uint32 NativeDecoderID = 0;
+		bool bUseVideoCodecSurface = false;
+		jobject VideoCodecSurface = nullptr; // global reference expected
+		bool bSurfaceIsView = false;
+	};
+
+	struct FDecoderInformation
+	{
+		int32 ApiLevel = 0;
+		bool bIsAdaptive = false;
+		bool bCanUse_SetOutputSurface = false;
+	};
+
+	struct FOutputFormatInfo
+	{
+		FOutputFormatInfo()
+		{
+			Clear();
+		}
+
+		void Clear()
+		{
+			Width = 0;
+			Height = 0;
+			CropTop = 0;
+			CropBottom = 0;
+			CropLeft = 0;
+			CropRight = 0;
+			Stride = 0;
+			SliceHeight = 0;
+			ColorFormat = 0;
+		}
+
+		bool IsValid() const
+		{
+			return Width != 0 && Height != 0;
+		}
+
+		int32 Width;
+		int32 Height;
+		int32 CropTop;
+		int32 CropBottom;
+		int32 CropLeft;
+		int32 CropRight;
+		int32 Stride;
+		int32 SliceHeight;
+		int32 ColorFormat;
+	};
+
+	struct FOutputBufferInfo
+	{
+		enum EBufferIndexValues
+		{
+			MediaCodec_INFO_TRY_AGAIN_LATER = -1,
+			MediaCodec_INFO_OUTPUT_FORMAT_CHANGED = -2,
+			MediaCodec_INFO_OUTPUT_BUFFERS_CHANGED = -3
+		};
+
+		FOutputBufferInfo()
+		{
+			Clear();
+		}
+
+		void Clear()
+		{
+			PresentationTimestamp = -1;
+			BufferIndex = -1;
+			Size = 0;
+			ValidCount = -1;
+			bIsEOS = false;
+			bIsConfig = false;
+		}
+
+		int64 PresentationTimestamp;
+		int32 BufferIndex;
+		int32 Size;
+		int32 ValidCount;
+		bool bIsEOS;
+		bool bIsConfig;
+	};
+
+
+	/**
+	 * Creates a Java instance of a VP8 or VP9 video decoder.
+	 *
+	 * @return 0 if successful, 1 on error.
+	 */
+	virtual int32 CreateDecoder(bool bCreateVP9) = 0;
+
+	/**
+	 * Initializes the video decoder instance.
+	 *
+	 * @param InCreateParams
+	 *
+	 * @return 0 if successful, 1 on error.
+	 */
+	virtual int32 InitializeDecoder(const FCreateParameters& InCreateParams) = 0;
+
+	/**
+	 * Attempts to set a new output surface on an existing and configured decoder.
+	 *
+	 * @param InNewOutputSurface
+	 *
+	 * @return 0 if successful, 1 on error.
+	 */
+	virtual int32 SetOutputSurface(jobject InNewOutputSurface) = 0;
+
+	/**
+	 * Releases (destroys) the Java video decoder instance.
+	 *
+	 * @return 0 if successful, 1 on error.
+	 */
+	virtual int32 ReleaseDecoder() = 0;
+
+	/**
+	 * Returns decoder information after a successful InitializeDecoder().
+	 *
+	 * @return Pointer to decoder information or null when no decoder has been created.
+	 */
+	virtual const FDecoderInformation* GetDecoderInformation() = 0;
+
+	/**
+	 * Starts the decoder instance.
+	 *
+	 * @return 0 if successful, 1 on error.
+	 */
+	virtual int32 Start() = 0;
+
+	/**
+	 * Stops the decoder instance.
+	 *
+	 * @return 0 if successful, 1 on error.
+	 */
+	virtual int32 Stop() = 0;
+
+	/**
+	 * Flushes the decoder instance.
+	 *
+	 * @return 0 if successful, 1 on error.
+	 */
+	virtual int32 Flush() = 0;
+
+	/**
+	 * Resets the decoder instance.
+	 *
+	 * @return 0 if successful, 1 on error.
+	 */
+	virtual int32 Reset() = 0;
+
+	/**
+	 * Dequeues an input buffer.
+	 *
+	 * @param InTimeoutUsec Timeout in microseconds to wait for an available buffer.
+	 *
+	 * @return >= 0 returns the index of the successfully dequeued buffer, negative values indicate an error.
+	 */
+	virtual int32 DequeueInputBuffer(int32 InTimeoutUsec) = 0;
+
+	/**
+	 * Queues input for decoding in the buffer with a previously dequeued (calling DequeueInputBuffer()) index.
+	 *
+	 * @param InBufferIndex Index of the buffer to put data into and enqueue for decoding (see DequeueInputBuffer()).
+	 * @param InPrefixData Data to prefix the data to decode with. If given it is usually the SPS & PPS.
+	 * @param InPrefixSize Size of the prefix data.
+	 * @param InAccessUnitData Data to be decoded.
+	 * @param InAccessUnitSize Size of the data to be decoded.
+	 * @param InTimestampUSec Timestamp (PTS) of the data, in microseconds.
+	 *
+	 * @return 0 if successful, 1 on error.
+	 */
+	virtual int32 QueueInputBuffer(int32 InBufferIndex, const void* InPrefixData, int32 InPrefixSize, const void* InAccessUnitData, int32 InAccessUnitSize, int64 InTimestampUSec) = 0;
+
+	/**
+	 * Queues codec specific data for the following to-be-decoded data buffers.
+	 *
+	 * @param InBufferIndex Index of the buffer to put data into and enqueue for decoding (see DequeueInputBuffer()).
+	 * @param InCSDData Codec specific data
+	 * @param InCSDSize Size of the codec specific data
+	 * @param InTimestampUSec Timestamp (PTS) of the data, in microseconds. Must be the same as the next data to be decoded.
+	 *
+	 * @return 0 if successful, 1 on error.
+	 */
+	virtual int32 QueueCSDInputBuffer(int32 InBufferIndex, const void* InCSDData, int32 InCSDSize, int64 InTimestampUSec) = 0;
+
+	/**
+	 * Queues end of stream for the buffer with a previously dequeued (calling DequeueInputBuffer()) index.
+	 *
+	 * @param InBufferIndex Index of the buffer to put the EOS flag into and enqueue for decoding (see DequeueInputBuffer()).
+	 * @param InTimestampUSec Timestamp the previous data had. Can be 0.
+	 *
+	 * @return 0 if successful, 1 on error.
+	 */
+	virtual int32 QueueEOSInputBuffer(int32 InBufferIndex, int64 InTimestampUSec) = 0;
+
+	/**
+	 * Returns format information of the decoded samples.
+	 *
+	 * @param OutFormatInfo
+	 * @param InOutputBufferIndex RESERVED FOR NOW - Pass any negative value to get the output format after DequeueOutputBuffer() returns a BufferIndex of MediaCodec_INFO_OUTPUT_FORMAT_CHANGED.
+	 *
+	 * @return 0 if successful, 1 on error.
+	 */
+	virtual int32 GetOutputFormatInfo(FOutputFormatInfo& OutFormatInfo, int32 InOutputBufferIndex) = 0;
+
+	/**
+	 * Dequeues an output buffer.
+	 *
+	 * @param InTimeoutUsec Timeout in microseconds to wait for an available buffer.
+	 *
+	 * @return 0 on success, 1 on failure. The OutBufferInfo.BufferIndex indicates the buffer index.
+	 */
+	virtual int32 DequeueOutputBuffer(FOutputBufferInfo& OutBufferInfo, int32 InTimeoutUsec) = 0;
+
+	/**
+	 * Returns the decoded samples from a decoder output buffer in the decoder native format.
+	 *
+	 * @param OutBufferDataPtr
+	 * @param OutBufferDataSize
+	 * @param InOutBufferInfo
+	 *
+	 * @return 0 on success, 1 on failure.
+	 */
+	virtual int32 GetOutputBuffer(void*& OutBufferDataPtr, int32 OutBufferDataSize, const FOutputBufferInfo& InOutBufferInfo) = 0;
+
+	/**
+	 * Releases the decoder output buffer back to the decoder.
+	 *
+	 * @return 0 on success, 1 on failure.
+	 */
+	virtual int32 ReleaseOutputBuffer(int32 BufferIndex, int32 ValidCount, bool bRender, int64 releaseAt) = 0;
+};
+
+#endif
