ElectraDecoders: Releasing decoded output buffer on Android as quickly as possible to have the OS propagate it into the surface queue as soon as possible

#rnx

diff --git a/Engine/Plugins/Media/ElectraCodecs/Source/ElectraDecoders/Private/Android/Java/src/com/epicgames/unreal/ElectraDecoderVideoH264.java b/Engine/Plugins/Media/ElectraCodecs/Source/ElectraDecoders/Private/Android/Java/src/com/epicgames/unreal/ElectraDecoderVideoH264.java
--- a/Engine/Plugins/Media/ElectraCodecs/Source/ElectraDecoders/Private/Android/Java/src/com/epicgames/unreal/ElectraDecoderVideoH264.java
+++ b/Engine/Plugins/Media/ElectraCodecs/Source/ElectraDecoders/Private/Android/Java/src/com/epicgames/unreal/ElectraDecoderVideoH264.java
@@ -143,7 +143,8 @@ public class ElectraDecoderVideoH264
 			MediaCodecList cl = new MediaCodecList(MediaCodecList.ALL_CODECS);
 			MediaCodecInfo[] codecInfos = cl.getCodecInfos();
 			int numCodecs = codecInfos.length;
-			for(int pass=0; pass<2; ++pass)
+//			for(int pass=0; pass<2; ++pass)
+			for(int pass=0; pass<1; ++pass)
 			{
 				for(int i=0; i<numCodecs; ++i)
 				{
@@ -171,12 +172,12 @@ public class ElectraDecoderVideoH264
 							boolean bSupportsLowLatency = DoesDecoderSupportsLowLatency(codecInfos[i], mimeType);
 
 							//Log.w(TAG, "["+pass+"] ElectraDecoderVideoH264: " + codecInfos[i].getCanonicalName() + "; ll="+bSupportsLowLatency);
-
+/*
 							if (pass==0 && !bSupportsLowLatency)
 							{
 								continue;
 							}
-
+*/
 							// Rumor has it that the codec list is sorted from "best" to "worst".
 							// Assuming this to be true we just return the first one matching the MIME type.
 							FCodecInformation Info = new FCodecInformation();
@@ -333,7 +334,7 @@ public class ElectraDecoderVideoH264
 				}
 			}
 			// Ask for low-latency.
-			Format.setInteger("low-latency", 1);
+//			Format.setInteger("low-latency", 1);
 
 
 			// Configure the decoder
diff --git a/Engine/Plugins/Media/ElectraCodecs/Source/ElectraDecoders/Private/Android/Java/src/com/epicgames/unreal/ElectraDecoderVideoH265.java b/Engine/Plugins/Media/ElectraCodecs/Source/ElectraDecoders/Private/Android/Java/src/com/epicgames/unreal/ElectraDecoderVideoH265.java
--- a/Engine/Plugins/Media/ElectraCodecs/Source/ElectraDecoders/Private/Android/Java/src/com/epicgames/unreal/ElectraDecoderVideoH265.java
+++ b/Engine/Plugins/Media/ElectraCodecs/Source/ElectraDecoders/Private/Android/Java/src/com/epicgames/unreal/ElectraDecoderVideoH265.java
@@ -142,7 +142,8 @@ public class ElectraDecoderVideoH265
 			MediaCodecList cl = new MediaCodecList(MediaCodecList.ALL_CODECS);
 			MediaCodecInfo[] codecInfos = cl.getCodecInfos();
 			int numCodecs = codecInfos.length;
-			for(int pass=0; pass<2; ++pass)
+//			for(int pass=0; pass<2; ++pass)
+			for(int pass=0; pass<1; ++pass)
 			{
 				for(int i=0; i<numCodecs; ++i)
 				{
@@ -170,12 +171,12 @@ public class ElectraDecoderVideoH265
 							boolean bSupportsLowLatency = DoesDecoderSupportsLowLatency(codecInfos[i], mimeType);
 
 							//Log.w(TAG, "["+pass+"] ElectraDecoderVideoH265: " + codecInfos[i].getCanonicalName() + "; ll="+bSupportsLowLatency);
-
+/*
 							if (pass==0 && !bSupportsLowLatency)
 							{
 								continue;
 							}
-
+*/
 							// Rumor has it that the codec list is sorted from "best" to "worst".
 							// Assuming this to be true we just return the first one matching the MIME type.
 							FCodecInformation Info = new FCodecInformation();
@@ -326,7 +327,7 @@ public class ElectraDecoderVideoH265
 				}
 			}
 			// Ask for low-latency.
-			Format.setInteger("low-latency", 1);
+//			Format.setInteger("low-latency", 1);
 
 
 			// Configure the decoder
diff --git a/Engine/Plugins/Media/ElectraCodecs/Source/ElectraDecoders/Private/Android/Java/src/com/epicgames/unreal/ElectraDecoderVideoVPx.java b/Engine/Plugins/Media/ElectraCodecs/Source/ElectraDecoders/Private/Android/Java/src/com/epicgames/unreal/ElectraDecoderVideoVPx.java
--- a/Engine/Plugins/Media/ElectraCodecs/Source/ElectraDecoders/Private/Android/Java/src/com/epicgames/unreal/ElectraDecoderVideoVPx.java
+++ b/Engine/Plugins/Media/ElectraCodecs/Source/ElectraDecoders/Private/Android/Java/src/com/epicgames/unreal/ElectraDecoderVideoVPx.java
@@ -144,7 +144,8 @@ public class ElectraDecoderVideoVPx
 			MediaCodecList cl = new MediaCodecList(MediaCodecList.ALL_CODECS);
 			MediaCodecInfo[] codecInfos = cl.getCodecInfos();
 			int numCodecs = codecInfos.length;
-			for(int pass=0; pass<2; ++pass)
+//			for(int pass=0; pass<2; ++pass)
+			for(int pass=0; pass<1; ++pass)
 			{
 				for(int i=0; i<numCodecs; ++i)
 				{
@@ -172,12 +173,12 @@ public class ElectraDecoderVideoVPx
 							boolean bSupportsLowLatency = DoesDecoderSupportsLowLatency(codecInfos[i], mimeType);
 
 							//Log.w(TAG, "["+pass+"] ElectraDecoderVideoVPx: " + codecInfos[i].getCanonicalName() + "; ll="+bSupportsLowLatency);
-
+/*
 							if (pass==0 && !bSupportsLowLatency)
 							{
 								continue;
 							}
-
+*/
 							// Rumor has it that the codec list is sorted from "best" to "worst".
 							// Assuming this to be true we just return the first one matching the MIME type.
 							FCodecInformation Info = new FCodecInformation();
@@ -339,7 +340,8 @@ public class ElectraDecoderVideoVPx
 				}
 			}
 			// Ask for low-latency.
-			Format.setInteger("low-latency", 1);
+//			Format.setInteger("low-latency", 1);
+
 
 			// Configure the decoder
 			android.view.Surface OutputSurface = null;
diff --git a/Engine/Plugins/Media/ElectraCodecs/Source/ElectraDecoders/Private/Android/h264/VideoDecoderH264_Android.cpp b/Engine/Plugins/Media/ElectraCodecs/Source/ElectraDecoders/Private/Android/h264/VideoDecoderH264_Android.cpp
--- a/Engine/Plugins/Media/ElectraCodecs/Source/ElectraDecoders/Private/Android/h264/VideoDecoderH264_Android.cpp
+++ b/Engine/Plugins/Media/ElectraCodecs/Source/ElectraDecoders/Private/Android/h264/VideoDecoderH264_Android.cpp
@@ -100,7 +100,7 @@ public:
 	TMap<FString, FVariant> ExtraValues;
 
 	TWeakPtr<FElectraVideoDecoderH264_Android, ESPMode::ThreadSafe> OwningDecoder;
-	IElectraH264VideoDecoderAndroidJava::FOutputBufferInfo OwningDecoderBufferInfo;
+	mutable IElectraH264VideoDecoderAndroidJava::FOutputBufferInfo OwningDecoderBufferInfo;
 	mutable bool bBufferGotReferenced = false;
 };
 
@@ -1097,6 +1097,7 @@ void FElectraVideoDecoderOutputH264_Android::ReleaseOutputBuffer()
 		{
 			Decoder->ReleaseOutputBuffer(&OwningDecoderBufferInfo, bBufferGotReferenced, -1);
 		}
+		OwningDecoderBufferInfo.BufferIndex = -1;
 	}
 }
 
@@ -1110,6 +1111,12 @@ IElectraDecoderVideoOutput::EImageCopyResult FElectraVideoDecoderOutputH264_Andr
 		InCopyResources->SetBufferIndex(OwningDecoderBufferInfo.BufferIndex);
 		InCopyResources->SetValidCount(OwningDecoderBufferInfo.ValidCount);
 
+		if (InCopyResources->ShouldReleaseBufferImmediately())
+		{
+			Decoder->ReleaseOutputBuffer(&OwningDecoderBufferInfo, bBufferGotReferenced, -1);
+			OwningDecoderBufferInfo.BufferIndex = -1;
+		}
+
 		return IElectraDecoderVideoOutput::EImageCopyResult::Succeeded; 
 	}
 	return IElectraDecoderVideoOutput::EImageCopyResult::Failed; 
diff --git a/Engine/Plugins/Media/ElectraCodecs/Source/ElectraDecoders/Private/Android/h265/VideoDecoderH265_Android.cpp b/Engine/Plugins/Media/ElectraCodecs/Source/ElectraDecoders/Private/Android/h265/VideoDecoderH265_Android.cpp
--- a/Engine/Plugins/Media/ElectraCodecs/Source/ElectraDecoders/Private/Android/h265/VideoDecoderH265_Android.cpp
+++ b/Engine/Plugins/Media/ElectraCodecs/Source/ElectraDecoders/Private/Android/h265/VideoDecoderH265_Android.cpp
@@ -100,7 +100,7 @@ public:
 	TMap<FString, FVariant> ExtraValues;
 
 	TWeakPtr<FElectraVideoDecoderH265_Android, ESPMode::ThreadSafe> OwningDecoder;
-	IElectraH265VideoDecoderAndroidJava::FOutputBufferInfo OwningDecoderBufferInfo;
+	mutable IElectraH265VideoDecoderAndroidJava::FOutputBufferInfo OwningDecoderBufferInfo;
 	mutable bool bBufferGotReferenced = false;
 };
 
@@ -1125,6 +1125,7 @@ void FElectraVideoDecoderOutputH265_Android::ReleaseOutputBuffer()
 		{
 			Decoder->ReleaseOutputBuffer(&OwningDecoderBufferInfo, bBufferGotReferenced, -1);
 		}
+		OwningDecoderBufferInfo.BufferIndex = -1;
 	}
 }
 
@@ -1138,6 +1139,12 @@ IElectraDecoderVideoOutput::EImageCopyResult FElectraVideoDecoderOutputH265_Andr
 		InCopyResources->SetBufferIndex(OwningDecoderBufferInfo.BufferIndex);
 		InCopyResources->SetValidCount(OwningDecoderBufferInfo.ValidCount);
 
+		if (InCopyResources->ShouldReleaseBufferImmediately())
+		{
+			Decoder->ReleaseOutputBuffer(&OwningDecoderBufferInfo, bBufferGotReferenced, -1);
+			OwningDecoderBufferInfo.BufferIndex = -1;
+		}
+
 		return IElectraDecoderVideoOutput::EImageCopyResult::Succeeded; 
 	}
 	return IElectraDecoderVideoOutput::EImageCopyResult::Failed; 
diff --git a/Engine/Plugins/Media/ElectraCodecs/Source/ElectraDecoders/Private/Android/vpx/VideoDecoderVPx_Android.cpp b/Engine/Plugins/Media/ElectraCodecs/Source/ElectraDecoders/Private/Android/vpx/VideoDecoderVPx_Android.cpp
--- a/Engine/Plugins/Media/ElectraCodecs/Source/ElectraDecoders/Private/Android/vpx/VideoDecoderVPx_Android.cpp
+++ b/Engine/Plugins/Media/ElectraCodecs/Source/ElectraDecoders/Private/Android/vpx/VideoDecoderVPx_Android.cpp
@@ -100,7 +100,7 @@ public:
 	TMap<FString, FVariant> ExtraValues;
 
 	TWeakPtr<FElectraVideoDecoderVPx_Android, ESPMode::ThreadSafe> OwningDecoder;
-	IElectraVPxVideoDecoderAndroidJava::FOutputBufferInfo OwningDecoderBufferInfo;
+	mutable IElectraVPxVideoDecoderAndroidJava::FOutputBufferInfo OwningDecoderBufferInfo;
 	mutable bool bBufferGotReferenced = false;
 };
 
@@ -1059,6 +1059,7 @@ void FElectraVideoDecoderOutputVPx_Android::ReleaseOutputBuffer()
 		{
 			Decoder->ReleaseOutputBuffer(&OwningDecoderBufferInfo, bBufferGotReferenced, -1);
 		}
+		OwningDecoderBufferInfo.BufferIndex = -1;
 	}
 }
 
@@ -1072,6 +1073,12 @@ IElectraDecoderVideoOutput::EImageCopyResult FElectraVideoDecoderOutputVPx_Andro
 		InCopyResources->SetBufferIndex(OwningDecoderBufferInfo.BufferIndex);
 		InCopyResources->SetValidCount(OwningDecoderBufferInfo.ValidCount);
 
+		if (InCopyResources->ShouldReleaseBufferImmediately())
+		{
+			Decoder->ReleaseOutputBuffer(&OwningDecoderBufferInfo, bBufferGotReferenced, -1);
+			OwningDecoderBufferInfo.BufferIndex = -1;
+		}
+
 		return IElectraDecoderVideoOutput::EImageCopyResult::Succeeded; 
 	}
 	return IElectraDecoderVideoOutput::EImageCopyResult::Failed; 
diff --git a/Engine/Plugins/Media/ElectraCodecs/Source/ElectraDecoders/Public/Android/AndroidElectraDecoderResourceDelegate.h b/Engine/Plugins/Media/ElectraCodecs/Source/ElectraDecoders/Public/Android/AndroidElectraDecoderResourceDelegate.h
--- a/Engine/Plugins/Media/ElectraCodecs/Source/ElectraDecoders/Public/Android/AndroidElectraDecoderResourceDelegate.h
+++ b/Engine/Plugins/Media/ElectraCodecs/Source/ElectraDecoders/Public/Android/AndroidElectraDecoderResourceDelegate.h
@@ -61,6 +61,7 @@ public:
 	virtual ~IElectraDecoderVideoOutputCopyResources() = default;
 	virtual void SetBufferIndex(int32 InIndex) = 0;
 	virtual void SetValidCount(int32 InValidCount) = 0;
+	virtual bool ShouldReleaseBufferImmediately() = 0;
 };
 
 
diff --git a/Engine/Plugins/Media/ElectraPlayer/Source/ElectraPlayerRuntime/Private/Runtime/Decoder/Android/AndroidElectraDecoderResourceManager.cpp b/Engine/Plugins/Media/ElectraPlayer/Source/ElectraPlayerRuntime/Private/Runtime/Decoder/Android/AndroidElectraDecoderResourceManager.cpp
--- a/Engine/Plugins/Media/ElectraPlayer/Source/ElectraPlayerRuntime/Private/Runtime/Decoder/Android/AndroidElectraDecoderResourceManager.cpp
+++ b/Engine/Plugins/Media/ElectraPlayer/Source/ElectraPlayerRuntime/Private/Runtime/Decoder/Android/AndroidElectraDecoderResourceManager.cpp
@@ -174,9 +174,12 @@ public:
 	{ BufferIndex = InIndex; }
 	void SetValidCount(int32 InValidCount) override
 	{ ValidCount = InValidCount; }
+	bool ShouldReleaseBufferImmediately() override
+	{ return bReleaseImmediately; }
 
 	int32 BufferIndex = -1;
 	int32 ValidCount = -1;
+	bool bReleaseImmediately = false;
 };
 
 
@@ -213,6 +216,7 @@ bool FElectraDecoderResourceManagerAndroid::SetupRenderBufferFromDecoderOutput(I
 		InOutBufferPropertes->Set(IDecoderOutputOptionNames::Pitch, FVariantValue((int64)InDecoderOutput->GetDecodedWidth() * ((NumBits > 8) ? 2 : 1)));
 
 		FElectraDecoderVideoOutputCopyResources cr;
+		cr.bReleaseImmediately = InInst->SurfaceType == IDecoderPlatformResourceAndroid::ISurfaceRequestCallback::ESurfaceType::Surface;
 		IElectraDecoderVideoOutput::EImageCopyResult CopyResult = InDecoderOutput->CopyPlatformImage(&cr);
 		if (CopyResult != IElectraDecoderVideoOutput::EImageCopyResult::Succeeded)
 		{
