Removing the need to create textures on-demand on the render thread from Xbox GDK playback code in Electra

diff --git a/Engine/Plugins/Media/ElectraCodecs/Source/ElectraDecoders/Private/DX/VideoDecoderH264_DX.cpp b/Engine/Plugins/Media/ElectraCodecs/Source/ElectraDecoders/Private/DX/VideoDecoderH264_DX.cpp
--- a/Engine/Plugins/Media/ElectraCodecs/Source/ElectraDecoders/Private/DX/VideoDecoderH264_DX.cpp
+++ b/Engine/Plugins/Media/ElectraCodecs/Source/ElectraDecoders/Private/DX/VideoDecoderH264_DX.cpp
@@ -898,8 +898,8 @@ bool FElectraVideoDecoderH264_DX::ConvertDecoderOutput()
 	NewOutput->DecodedWidth = (int32)dwInputWidth;
 	NewOutput->DecodedHeight = (int32)dwInputHeight;
 	// If we deliver the data in a CPU side buffer, we need to patch the decoded height to reflect the oddness of these planar formats
-	// (DirectX12 and higher, SW decoder or a non-DirectX API will trigger this)
-	if (DecoderPlatformHandle->IsSoftware() || DecoderPlatformHandle->GetDXVersionTimes1000() >= 12000 || DecoderPlatformHandle->GetDXVersionTimes1000() == 0)
+	// (DirectX12 and higher (unless we can use the render device to decode), SW decoder or a non-DirectX API will trigger this)
+	if (DecoderPlatformHandle->IsSoftware() || (DecoderPlatformHandle->GetDXVersionTimes1000() >= 12000 && DecoderPlatformHandle->GetDXDevice() != nullptr) || DecoderPlatformHandle->GetDXVersionTimes1000() == 0)
 	{
 		// We are returning a CPU side buffer (software decode OR DX12 render device) - adjust the height so it can be interpreted as a single plane texture
 		NewOutput->DecodedHeight = NewOutput->DecodedHeight * 3 / 2;
diff --git a/Engine/Plugins/Media/ElectraCodecs/Source/ElectraDecoders/Private/DX/VideoDecoderH265_DX.cpp b/Engine/Plugins/Media/ElectraCodecs/Source/ElectraDecoders/Private/DX/VideoDecoderH265_DX.cpp
--- a/Engine/Plugins/Media/ElectraCodecs/Source/ElectraDecoders/Private/DX/VideoDecoderH265_DX.cpp
+++ b/Engine/Plugins/Media/ElectraCodecs/Source/ElectraDecoders/Private/DX/VideoDecoderH265_DX.cpp
@@ -911,7 +911,9 @@ bool FElectraVideoDecoderH265_DX::ConvertDecoderOutput()
 	// Width and height as "decoded dims" (which are featuring format specific scales & no cropping)
 	NewOutput->DecodedWidth = (int32)dwInputWidth;
 	NewOutput->DecodedHeight = (int32)dwInputHeight;
-	if (DecoderPlatformHandle->IsSoftware() || DecoderPlatformHandle->GetDXVersionTimes1000() >= 12000 || DecoderPlatformHandle->GetDXVersionTimes1000() == 0)
+	// If we deliver the data in a CPU side buffer, we need to patch the decoded height to reflect the oddness of these planar formats
+	// (DirectX12 and higher (unless we can use the render device to decode), SW decoder or a non-DirectX API will trigger this)
+	if (DecoderPlatformHandle->IsSoftware() || (DecoderPlatformHandle->GetDXVersionTimes1000() >= 12000 && DecoderPlatformHandle->GetDXDevice() != nullptr) || DecoderPlatformHandle->GetDXVersionTimes1000() == 0)
 	{
 		// We are returning a CPU side buffer (software decode OR DX12 render device) - adjust the height so it can be interpreted as a single plane texture
 		NewOutput->DecodedHeight = NewOutput->DecodedHeight * 3 / 2;
