Updating Electra output pipeline when used with DX12. Implements RHI independent data uploads, usage of new DX12 support in WMF in SDK 22621 and avoids fijnal copy to establish an RHI texture by creating one from a native resource
#tests Local test on PC, XSX and XB1 as well as OSX

diff --git a/Engine/Plugins/Experimental/TextureMediaPlayer/Source/TextureMediaPlayer/Private/TextureMediaPlayerVideoDecoderOutput.cpp b/Engine/Plugins/Experimental/TextureMediaPlayer/Source/TextureMediaPlayer/Private/TextureMediaPlayerVideoDecoderOutput.cpp
--- a/Engine/Plugins/Experimental/TextureMediaPlayer/Source/TextureMediaPlayer/Private/TextureMediaPlayerVideoDecoderOutput.cpp
+++ b/Engine/Plugins/Experimental/TextureMediaPlayer/Source/TextureMediaPlayer/Private/TextureMediaPlayerVideoDecoderOutput.cpp
@@ -54,12 +54,6 @@ FIntPoint TextureMediaPlayerVideoDecoderOutput::GetDim() const
 
 #if PLATFORM_WINDOWS
 
-TRefCountPtr<IMFSample> TextureMediaPlayerVideoDecoderOutput::GetMFSample() const
-{
-	check(!"Should not be called!");
-	return nullptr;
-}
-
 TRefCountPtr<IUnknown> TextureMediaPlayerVideoDecoderOutput::GetTexture() const
 {
 	IUnknown* Result = nullptr;
diff --git a/Engine/Plugins/Experimental/TextureMediaPlayer/Source/TextureMediaPlayer/Private/TextureMediaPlayerVideoDecoderOutput.h b/Engine/Plugins/Experimental/TextureMediaPlayer/Source/TextureMediaPlayer/Private/TextureMediaPlayerVideoDecoderOutput.h
--- a/Engine/Plugins/Experimental/TextureMediaPlayer/Source/TextureMediaPlayer/Private/TextureMediaPlayerVideoDecoderOutput.h
+++ b/Engine/Plugins/Experimental/TextureMediaPlayer/Source/TextureMediaPlayer/Private/TextureMediaPlayerVideoDecoderOutput.h
@@ -34,8 +34,6 @@ public:
 	virtual FIntPoint GetDim() const;
 
 #if PLATFORM_WINDOWS
-	virtual TRefCountPtr<IMFSample> GetMFSample() const override;
-
 	virtual TRefCountPtr<IUnknown> GetTexture() const override;
 	virtual TRefCountPtr<IUnknown> GetSync(uint64& SyncValue) const override;
 	virtual TRefCountPtr<ID3D11Device> GetDevice() const override;
diff --git a/Engine/Plugins/Media/AppleProResDecoderElectra/Source/AppleProResDecoderElectra/AppleProResDecoderElectra.Build.cs b/Engine/Plugins/Media/AppleProResDecoderElectra/Source/AppleProResDecoderElectra/AppleProResDecoderElectra.Build.cs
--- a/Engine/Plugins/Media/AppleProResDecoderElectra/Source/AppleProResDecoderElectra/AppleProResDecoderElectra.Build.cs
+++ b/Engine/Plugins/Media/AppleProResDecoderElectra/Source/AppleProResDecoderElectra/AppleProResDecoderElectra.Build.cs
@@ -1,5 +1,7 @@
 // Copyright Epic Games, Inc. All Rights Reserved.
 
+using System.IO;
+
 namespace UnrealBuildTool.Rules
 {
 	public class AppleProResDecoderElectra : ModuleRules
@@ -20,6 +22,16 @@ namespace UnrealBuildTool.Rules
                     new string[] {
                     "ProResLib",
                     });
+
+				PublicDependencyModuleNames.Add("DirectX");
+				AddEngineThirdPartyPrivateStaticDependencies(Target, "DX12");
+
+				if (Target.WindowsPlatform.Architecture != UnrealArch.Arm64)
+				{
+					PublicAdditionalLibraries.AddRange(new string[] {
+						Path.Combine(Target.WindowsPlatform.DirectXLibDir, "dxerr.lib"),
+					});
+				}
 			}
 			else if (Target.Platform == UnrealTargetPlatform.Mac)
 			{
diff --git a/Engine/Plugins/Media/AppleProResDecoderElectra/Source/AppleProResDecoderElectra/Private/ProResDecoder/Mac/ElectraMediaProResDecoder.cpp b/Engine/Plugins/Media/AppleProResDecoderElectra/Source/AppleProResDecoderElectra/Private/ProResDecoder/Mac/ElectraMediaProResDecoder.cpp
--- a/Engine/Plugins/Media/AppleProResDecoderElectra/Source/AppleProResDecoderElectra/Private/ProResDecoder/Mac/ElectraMediaProResDecoder.cpp
+++ b/Engine/Plugins/Media/AppleProResDecoderElectra/Source/AppleProResDecoderElectra/Private/ProResDecoder/Mac/ElectraMediaProResDecoder.cpp
@@ -114,6 +114,10 @@ public:
 	{
 		return InBufferIndex == 0 ? ImageBuffer : nullptr;
 	}
+	virtual bool GetBufferTextureSyncByIndex(int32 InBufferIndex, FElectraDecoderOutputSync& SyncObject) const override
+	{
+		return false;
+	}
 	EElectraDecoderPlatformPixelFormat GetBufferFormatByIndex(int32 InBufferIndex) const override
 	{
 		return PixelFormat;
diff --git a/Engine/Plugins/Media/AppleProResDecoderElectra/Source/AppleProResDecoderElectra/Private/ProResDecoder/Windows/ElectraMediaProResDecoder.cpp b/Engine/Plugins/Media/AppleProResDecoderElectra/Source/AppleProResDecoderElectra/Private/ProResDecoder/Windows/ElectraMediaProResDecoder.cpp
--- a/Engine/Plugins/Media/AppleProResDecoderElectra/Source/AppleProResDecoderElectra/Private/ProResDecoder/Windows/ElectraMediaProResDecoder.cpp
+++ b/Engine/Plugins/Media/AppleProResDecoderElectra/Source/AppleProResDecoderElectra/Private/ProResDecoder/Windows/ElectraMediaProResDecoder.cpp
@@ -15,6 +15,8 @@
 #include "IElectraDecoderResourceDelegate.h"
 #include "ElectraDecodersUtils.h"
 #include "Utils/MPEG/ElectraUtilsMP4.h"	// to read the ProRes header, which is given in the form of an MP4 Atom
+#include COMPILED_PLATFORM_HEADER(ElectraDecoderGPUBufferHelpers.h)
+
 
 #include "ProResDecoder.h"
 
@@ -130,8 +132,25 @@ public:
 	}
 	void* GetBufferTextureByIndex(int32 InBufferIndex) const override
 	{
+#if ELECTRA_MEDIAGPUBUFFER_DX12
+		if (InBufferIndex == 0)
+		{
+			return GPUBuffer.Resource.GetReference();
+		}
+#endif
 		return nullptr;
 	}
+	virtual bool GetBufferTextureSyncByIndex(int32 InBufferIndex, FElectraDecoderOutputSync& SyncObject) const override
+	{
+#if ELECTRA_MEDIAGPUBUFFER_DX12
+		if (InBufferIndex == 0)
+		{
+			SyncObject = { GPUBuffer.Fence.GetReference(), GPUBuffer.FenceValue };
+			return true;
+		}
+#endif
+		return false;
+	}
 	EElectraDecoderPlatformPixelFormat GetBufferFormatByIndex(int32 InBufferIndex) const override
 	{
 		return BufferFormat;
@@ -167,6 +186,9 @@ public:
 	TSharedPtr<TArray<uint8>, ESPMode::ThreadSafe> Buffer;
 	EElectraDecoderPlatformPixelFormat BufferFormat = EElectraDecoderPlatformPixelFormat::INVALID;
 	EElectraDecoderPlatformPixelEncoding BufferEncoding = EElectraDecoderPlatformPixelEncoding::Native;
+#if ELECTRA_MEDIAGPUBUFFER_DX12
+	FElectraMediaDecoderOutputBufferPool_DX12::FOutputData GPUBuffer;
+#endif
 };
 
 
@@ -346,6 +368,8 @@ private:
 	uint32 Codec4CC = 0;
 	uint32 BufferAllocationSize = 0;
 
+	TWeakPtr<IElectraDecoderResourceDelegate, ESPMode::ThreadSafe> ResourceDelegate;
+
 	IElectraDecoder::FError LastError;
 
 	PRPixelFormat OutputPixelFormat = kPRFormat_y416;
@@ -357,6 +381,11 @@ private:
 
 	PRDecoderRef Decoder = nullptr;
 	FProResHeader CurrentHeader;
+
+	uint32 MaxOutputBuffers;
+#if ELECTRA_MEDIAGPUBUFFER_DX12
+	mutable TSharedPtr<FElectraMediaDecoderOutputBufferPool_DX12> D3D12ResourcePool;
+#endif
 };
 
 
@@ -432,6 +461,10 @@ FVideoDecoderProResElectra::FVideoDecoderProResElectra(const TMap<FString, FVari
 	// The decoded width is the same as the display width
 	DecodedWidth = DisplayWidth;
 	DecodedHeight = DisplayHeight;
+	ResourceDelegate = InResourceDelegate;
+
+	MaxOutputBuffers = (uint32)ElectraDecodersUtil::GetVariantValueSafeU64(InOptions, TEXT("max_output_buffers"), 5);
+	MaxOutputBuffers += kElectraDecoderPipelineExtraFrames;
 }
 
 FVideoDecoderProResElectra::~FVideoDecoderProResElectra()
@@ -587,6 +620,14 @@ IElectraDecoder::EDecoderError FVideoDecoderProResElectra::DecodeAccessUnit(cons
 		return IElectraDecoder::EDecoderError::NoBuffer;
 	}
 
+#if ELECTRA_MEDIAGPUBUFFER_DX12
+	// If we will create a new resource pool or we have still buffers in an existing one, we can proceed, else we'd have no resources to output the data
+	if (D3D12ResourcePool.IsValid() && !D3D12ResourcePool->BufferAvailable())
+	{
+		return IElectraDecoder::EDecoderError::NoBuffer;
+	}
+#endif
+
 	// Create a decoder?
 	if (!Decoder)
 	{
@@ -619,6 +660,17 @@ IElectraDecoder::EDecoderError FVideoDecoderProResElectra::DecodeAccessUnit(cons
 			return IElectraDecoder::EDecoderError::Error;
 		}
 
+		void* PlatformDevice = nullptr;
+		int32 PlatformDeviceVersion = 0;
+		bool bUseGPUBuffers = false;
+#if ELECTRA_MEDIAGPUBUFFER_DX12
+		if (auto PinnedResourceDelegate = ResourceDelegate.Pin())
+		{
+			PinnedResourceDelegate->GetD3DDevice(&PlatformDevice, &PlatformDeviceVersion);
+			bUseGPUBuffers = (PlatformDevice && PlatformDeviceVersion >= 12000);
+		}
+#endif
+
 		TSharedPtr<FVideoDecoderOutputProResElectra, ESPMode::ThreadSafe> NewOutput = MakeShared<FVideoDecoderOutputProResElectra>();
 		NewOutput->PTS = InInputAccessUnit.PTS;
 		NewOutput->UserValue = InInputAccessUnit.UserValue;
@@ -691,18 +743,67 @@ IElectraDecoder::EDecoderError FVideoDecoderProResElectra::DecodeAccessUnit(cons
 		NewOutput->ExtraValues.Emplace(TEXT("codec"), FVariant(TEXT("prores")));
 		NewOutput->ExtraValues.Emplace(TEXT("codec_4cc"), FVariant(Codec4CC));
 
-		NewOutput->Buffer = MakeShared<TArray<uint8>, ESPMode::ThreadSafe>();
-		NewOutput->Buffer->AddUninitialized(BufferAllocationSize);
-
 		PRPixelBuffer pb;
 		FMemory::Memzero(pb);
-		pb.baseAddr = NewOutput->Buffer->GetData();
-		pb.rowBytes = OutputBufferBytesPerRow;
 		pb.format = OutputPixelFormat;
 		pb.width = DecodedWidth;
 		pb.height = DecodedHeight;
 
+#if ELECTRA_MEDIAGPUBUFFER_DX12
+		if (bUseGPUBuffers)
+		{
+			// Setup a resource to directly receive the decoder output
+			// (the memory is WC configured, so the decoder must not read from it -- this seems to be the case with ProRes)
+
+			TRefCountPtr D3D12Device(static_cast<ID3D12Device*>(PlatformDevice));
+
+			// Create the resource pool as needed...
+			if (!D3D12ResourcePool)
+			{
+				// note: prores will use a constant resolution and format throughout a decoder session
+				D3D12ResourcePool = MakeShared<FElectraMediaDecoderOutputBufferPool_DX12, ESPMode::ThreadSafe>(D3D12Device, MaxOutputBuffers, NewOutput->Width, NewOutput->Height, NewOutput->Pitch / NewOutput->Width);
+			}
+
+			// Request resource and fence...
+			uint32 BufferPitch;
+			D3D12ResourcePool->AllocateOutputDataAsBuffer(NewOutput->GPUBuffer, BufferPitch);
+
+			// Correct output pitch to what the resource is setup for
+			NewOutput->Pitch = BufferPitch;
+
+			void* BufferAddr;
+			HRESULT Res = NewOutput->GPUBuffer.Resource->Map(0, nullptr, &BufferAddr);
+			check(SUCCEEDED(Res));
+
+			check((uint32)OutputBufferBytesPerRow <= BufferPitch);
+			pb.baseAddr = (unsigned char*)BufferAddr;
+			pb.rowBytes = BufferPitch;
+		}
+		else
+#endif
+		{
+			NewOutput->Buffer = MakeShared<TArray<uint8>, ESPMode::ThreadSafe>();
+			NewOutput->Buffer->AddUninitialized(BufferAllocationSize);
+
+			pb.baseAddr = NewOutput->Buffer->GetData();
+			pb.rowBytes = OutputBufferBytesPerRow;
+		}
+
 		int NumBytesDecoded = PRDecodeFrame(Decoder, InInputAccessUnit.Data, static_cast<int>(InInputAccessUnit.DataSize), &pb, OutDownscaleMode, OutputDiscardAlpha);
+
+#if ELECTRA_MEDIAGPUBUFFER_DX12
+		if (bUseGPUBuffers)
+		{
+			// Unmap the resource memory and signal that it's usable by the GPU
+			NewOutput->GPUBuffer.Resource->Unmap(0, nullptr);
+
+			if (NumBytesDecoded >= 0)
+			{
+				NewOutput->GPUBuffer.Fence->Signal(NewOutput->GPUBuffer.FenceValue);
+			}
+		}
+#endif
+
 		if (NumBytesDecoded < 0)
 		{
 			PostError(NumBytesDecoded, TEXT("PRDecodeFrame() failed"), ERRCODE_INTERNAL_FAILED_TO_DECODE_INPUT);
diff --git a/Engine/Plugins/Media/AvidDNxHDDecoderElectra/Source/AvidDNxHDDecoderElectra.Build.cs b/Engine/Plugins/Media/AvidDNxHDDecoderElectra/Source/AvidDNxHDDecoderElectra.Build.cs
--- a/Engine/Plugins/Media/AvidDNxHDDecoderElectra/Source/AvidDNxHDDecoderElectra.Build.cs
+++ b/Engine/Plugins/Media/AvidDNxHDDecoderElectra/Source/AvidDNxHDDecoderElectra.Build.cs
@@ -1,5 +1,7 @@
 // Copyright Epic Games, Inc. All Rights Reserved.
 
+using System.IO;
+
 namespace UnrealBuildTool.Rules
 {
 	public class AvidDNxHDDecoderElectra : ModuleRules
@@ -23,6 +25,16 @@ namespace UnrealBuildTool.Rules
 						//"DNxUncompressed",
 					}
 				);
+
+				PublicDependencyModuleNames.Add("DirectX");
+				AddEngineThirdPartyPrivateStaticDependencies(Target, "DX12");
+
+				if (Target.WindowsPlatform.Architecture != UnrealArch.Arm64)
+				{
+					PublicAdditionalLibraries.AddRange(new string[] {
+						Path.Combine(Target.WindowsPlatform.DirectXLibDir, "dxerr.lib"),
+					});
+				}
 			}
 		}
 	}
diff --git a/Engine/Plugins/Media/AvidDNxHDDecoderElectra/Source/Private/Windows/ElectraMediaAvidDNxHDDecoder.cpp b/Engine/Plugins/Media/AvidDNxHDDecoderElectra/Source/Private/Windows/ElectraMediaAvidDNxHDDecoder.cpp
--- a/Engine/Plugins/Media/AvidDNxHDDecoderElectra/Source/Private/Windows/ElectraMediaAvidDNxHDDecoder.cpp
+++ b/Engine/Plugins/Media/AvidDNxHDDecoderElectra/Source/Private/Windows/ElectraMediaAvidDNxHDDecoder.cpp
@@ -18,6 +18,7 @@
 #include "IElectraDecoderOutputVideo.h"
 #include "IElectraDecoderResourceDelegate.h"
 #include "ElectraDecodersUtils.h"
+#include COMPILED_PLATFORM_HEADER(ElectraDecoderGPUBufferHelpers.h)
 
 #include "AvidDNxCodec.h"
 //#include "dnx_uncompressed_sdk.h"
@@ -138,8 +139,25 @@ public:
 	}
 	void* GetBufferTextureByIndex(int32 InBufferIndex) const override
 	{
+#if ELECTRA_MEDIAGPUBUFFER_DX12
+		if (InBufferIndex == 0)
+		{
+			return GPUBuffer.Resource.GetReference();
+		}
+#endif
 		return nullptr;
 	}
+	virtual bool GetBufferTextureSyncByIndex(int32 InBufferIndex, FElectraDecoderOutputSync& SyncObject) const override
+	{
+#if ELECTRA_MEDIAGPUBUFFER_DX12
+		if (InBufferIndex == 0)
+		{
+			SyncObject = { GPUBuffer.Fence.GetReference(), GPUBuffer.FenceValue };
+			return true;
+		}
+#endif
+	return false;
+	}
 	EElectraDecoderPlatformPixelFormat GetBufferFormatByIndex(int32 InBufferIndex) const override
 	{
 		if (InBufferIndex == 0)
@@ -187,6 +205,9 @@ public:
 	TSharedPtr<TArray<uint8>, ESPMode::ThreadSafe> Buffer;
 	EElectraDecoderPlatformPixelFormat BufferFormat = EElectraDecoderPlatformPixelFormat::INVALID;
 	EElectraDecoderPlatformPixelEncoding BufferEncoding = EElectraDecoderPlatformPixelEncoding::Native;
+#if ELECTRA_MEDIAGPUBUFFER_DX12
+	FElectraMediaDecoderOutputBufferPool_DX12::FOutputData GPUBuffer;
+#endif
 };
 
 
@@ -249,22 +270,17 @@ private:
 	struct FDecoderHandle
 	{
 		int32 PrepareForDecoding(const void* InInput, uint32 InInputSize, int32 InNumDecodeThreads, bool bIgnoreAlpha);
-		int32 GetDecodedBufferSize() const
-		{ return BufferAllocSize; }
 		int32 GetNumOutputBits() const
 		{ return NumOutputBits; }
-		int32 Decode(TSharedPtr<TArray<uint8>, ESPMode::ThreadSafe> OutBuffer, const void* InInput, uint32 InInputSize);
+		int32 Decode(uint8* OutBufferData, uint32 OutBufferSize, const void* InInput, uint32 InInputSize);
 		void Reset();
 		bool SetupDecompressStruct();
-		static int32 GetNumberOfOutputColorComponents(DNX_ColorComponentOrder_t InOutputComponentOrder);
-		static int32 GetNumberOfOutputColorBytesPerPixel(DNX_ComponentType_t InOutputComponentType);
 		static int32 GetNumberOfOutputColorBitsPerPixel(DNX_ComponentType_t InOutputComponentType);
 		DNX_Decoder Handle = nullptr;
 		DNX_CompressedParams_t CurrentCompressedParams;
 		DNX_UncompressedParams_t CurrentUncompressedParams;
 		DNX_SignalStandard_t CurrentSignalStandard = DNX_SignalStandard_t::DNX_SS_INVALID;
 		DNX_DecodeOperationParams_t DecodeOperationParams;
-		uint32 BufferAllocSize = 0;
 		int32 NumOutputBits = 0;
 		bool bHasAlpha = false;
 	};
@@ -275,6 +291,14 @@ private:
 
 	TSharedPtr<FVideoDecoderOutputAvidDNxHDElectra, ESPMode::ThreadSafe> CurrentOutput;
 	bool bFlushPending = false;
+
+	TWeakPtr<IElectraDecoderResourceDelegate, ESPMode::ThreadSafe> ResourceDelegate;
+
+	uint32 MaxOutputBuffers;
+#if ELECTRA_MEDIAGPUBUFFER_DX12
+	mutable TSharedPtr<FElectraMediaDecoderOutputBufferPool_DX12> D3D12ResourcePool;
+	mutable TUniquePtr<IElectraDecoderResourceDelegateBase::IAsyncConsecutiveTaskSync> TaskSync;
+#endif
 };
 
 
@@ -388,6 +412,8 @@ void FElectraMediaAvidDNxHDDecoder::Shutdown()
 
 FVideoDecoderAvidDNxHDElectra::FVideoDecoderAvidDNxHDElectra(const TMap<FString, FVariant>& InOptions, TSharedPtr<IElectraDecoderResourceDelegate, ESPMode::ThreadSafe> InResourceDelegate)
 {
+	ResourceDelegate = InResourceDelegate;
+
 	DisplayWidth = (int32)ElectraDecodersUtil::GetVariantValueSafeI64(InOptions, TEXT("width"), 0);
 	DisplayHeight = (int32)ElectraDecodersUtil::GetVariantValueSafeI64(InOptions, TEXT("height"), 0);
 
@@ -399,6 +425,9 @@ FVideoDecoderAvidDNxHDElectra::FVideoDecoderAvidDNxHDElectra(const TMap<FString,
 
 	// FIXME: How many exactly?
 	DesiredNumberOfDecoderThreads = 8;
+
+	MaxOutputBuffers = (uint32)ElectraDecodersUtil::GetVariantValueSafeU64(InOptions, TEXT("max_output_buffers"), 5);
+	MaxOutputBuffers += kElectraDecoderPipelineExtraFrames;
 }
 
 FVideoDecoderAvidDNxHDElectra::~FVideoDecoderAvidDNxHDElectra()
@@ -455,6 +484,18 @@ TSharedPtr<IElectraDecoderDefaultOutputFormat, ESPMode::ThreadSafe> FVideoDecode
 	return nullptr;
 }
 
+static void CopyData(TRefCountPtr<ID3D12Resource> BufferResource, uint32 BufferPitch, const uint8* TempBuffer, uint32 TempBufferPitch, uint32 Height)
+{
+	uint8* ResourceAddr = nullptr;
+	HRESULT Res = BufferResource->Map(0, nullptr, (void**)&ResourceAddr);
+	check(SUCCEEDED(Res));
+
+	FElectraMediaDecoderOutputBufferPool_DX12::CopyWithPitchAdjust(ResourceAddr, BufferPitch, TempBuffer, TempBufferPitch, Height);
+
+	// We decoded into the resource: Unmap the resource memory and signal that it's usable by the GPU
+	BufferResource->Unmap(0, nullptr);
+}
+
 IElectraDecoder::EDecoderError FVideoDecoderAvidDNxHDElectra::DecodeAccessUnit(const FInputAccessUnit& InInputAccessUnit, const TMap<FString, FVariant>& InAdditionalOptions)
 {
 	// If already in error do nothing!
@@ -476,6 +517,14 @@ IElectraDecoder::EDecoderError FVideoDecoderAvidDNxHDElectra::DecodeAccessUnit(c
 		return IElectraDecoder::EDecoderError::NoBuffer;
 	}
 
+#if ELECTRA_MEDIAGPUBUFFER_DX12
+	// If we will create a new resource pool or we have still buffers in an existing one, we can proceed, else we'd have no resources to output the data
+	if (D3D12ResourcePool.IsValid() && !D3D12ResourcePool->BufferAvailable())
+	{
+		return IElectraDecoder::EDecoderError::NoBuffer;
+	}
+#endif
+
 	// Decode data. This immediately produces a new output frame.
 	if (InInputAccessUnit.Data && InInputAccessUnit.DataSize)
 	{
@@ -487,6 +536,18 @@ IElectraDecoder::EDecoderError FVideoDecoderAvidDNxHDElectra::DecodeAccessUnit(c
 			return IElectraDecoder::EDecoderError::Error;
 		}
 
+		void* PlatformDevice = nullptr;
+		int32 PlatformDeviceVersion = 0;
+		bool bUseGPUBuffers = false;
+		auto PinnedResourceDelegate = ResourceDelegate.Pin();
+#if ELECTRA_MEDIAGPUBUFFER_DX12
+		if (PinnedResourceDelegate.IsValid())
+		{
+			PinnedResourceDelegate->GetD3DDevice(&PlatformDevice, &PlatformDeviceVersion);
+			bUseGPUBuffers = (PlatformDevice && PlatformDeviceVersion >= 12000);
+		}
+#endif
+
 		TSharedPtr<FVideoDecoderOutputAvidDNxHDElectra, ESPMode::ThreadSafe> NewOutput = MakeShared<FVideoDecoderOutputAvidDNxHDElectra>();
 		NewOutput->PTS = InInputAccessUnit.PTS;
 		NewOutput->UserValue = InInputAccessUnit.UserValue;
@@ -509,6 +570,7 @@ IElectraDecoder::EDecoderError FVideoDecoderAvidDNxHDElectra::DecodeAccessUnit(c
 
 		NewOutput->PixelFormat = (int32) Decoder.CurrentUncompressedParams.compType;
 
+		bool bIsPlanar = false;
 		switch (Decoder.CurrentUncompressedParams.compType)
 		{
 			case	DNX_CT_UCHAR:
@@ -549,7 +611,11 @@ IElectraDecoder::EDecoderError FVideoDecoderAvidDNxHDElectra::DecodeAccessUnit(c
 					{
 						NewOutput->BufferFormat = EElectraDecoderPlatformPixelFormat::NV12;
 						NewOutput->BufferEncoding = EElectraDecoderPlatformPixelEncoding::Native;
-						NewOutput->Height = (NewOutput->Height * 3) / 2;
+						if (!bUseGPUBuffers)
+						{
+							NewOutput->Height = (NewOutput->Height * 3) / 2;
+						}
+						bIsPlanar = true;
 						break;
 					}
 					default:
@@ -600,7 +666,11 @@ IElectraDecoder::EDecoderError FVideoDecoderAvidDNxHDElectra::DecodeAccessUnit(c
 					{
 						NewOutput->BufferFormat = EElectraDecoderPlatformPixelFormat::P010;
 						NewOutput->BufferEncoding = EElectraDecoderPlatformPixelEncoding::Native;
-						NewOutput->Height = (NewOutput->Height * 3) / 2;
+						if (!bUseGPUBuffers)
+						{
+							NewOutput->Height = (NewOutput->Height * 3) / 2;
+						}
+						bIsPlanar = true;
 						break;
 					}
 					default:
@@ -629,15 +699,106 @@ IElectraDecoder::EDecoderError FVideoDecoderAvidDNxHDElectra::DecodeAccessUnit(c
 
 		NewOutput->ExtraValues.Emplace(TEXT("avid_color_volume"), FVariant((int32) Decoder.CurrentUncompressedParams.colorVolume));
 
-		NewOutput->Buffer = MakeShared<TArray<uint8>, ESPMode::ThreadSafe>();
-		NewOutput->Buffer->AddUninitialized(Decoder.GetDecodedBufferSize());
+		uint32 BPP = NewOutput->Pitch / NewOutput->Width;
+
+		uint8* OutBufferData;
+		uint32 OutBufferSize;
+		TArray<uint8> TempBuffer;
+
+#if ELECTRA_MEDIAGPUBUFFER_DX12
+		uint32 ResourcePitch = 0;
+		if (bUseGPUBuffers)
+		{
+			TRefCountPtr D3D12Device(static_cast<ID3D12Device*>(PlatformDevice));
+
+			uint32 AllocHeight = bIsPlanar ? (NewOutput->Height * 3 / 2) : NewOutput->Height;
+
+			// Create the resource pool as needed...
+			if (!D3D12ResourcePool || !D3D12ResourcePool->IsCompatibleAsBuffer(MaxOutputBuffers, NewOutput->Width, AllocHeight, BPP))
+			{
+				// Note: if we reconfigure the pool, we will end up with >1 heaps until all older users are destroyed - so streams with lots of changes will be quite resource hungry
+				D3D12ResourcePool = MakeShared<FElectraMediaDecoderOutputBufferPool_DX12, ESPMode::ThreadSafe>(D3D12Device, MaxOutputBuffers, NewOutput->Width, AllocHeight, BPP);
+			}
+
+			// Create a tasksync instance so we can have our own async jobs run in consecutive order
+			if (!TaskSync.IsValid())
+			{
+				TaskSync.Reset(PinnedResourceDelegate->CreateAsyncConsecutiveTaskSync());
+			}
+
+			// Request resource and fence...
+			D3D12ResourcePool->AllocateOutputDataAsBuffer(NewOutput->GPUBuffer, ResourcePitch);
+
+			NewOutput->Pitch = ResourcePitch;
+
+			uint32 OutputPitch = BPP * NewOutput->Width;
+
+			// Pitch compatible with DX12 resource?
+			if (OutputPitch != ResourcePitch)
+			{
+				// No, decode to temp buffer...
+				TempBuffer.AddUninitialized(OutputPitch * AllocHeight);
+
+				OutBufferData = TempBuffer.GetData();
+				OutBufferSize = TempBuffer.Num();
+
+				check(OutBufferSize <= AllocHeight * ResourcePitch);
+			}
+			else
+			{
+				// Yes, decode directly into the DX12 resource...
+				HRESULT Res = NewOutput->GPUBuffer.Resource->Map(0, nullptr, (void**)&OutBufferData);
+				check(SUCCEEDED(Res));
+				OutBufferSize = ResourcePitch * AllocHeight;
+			}
+		}
+		else
+#endif
+		{
+			NewOutput->Buffer = MakeShared<TArray<uint8>, ESPMode::ThreadSafe>();
+			NewOutput->Buffer->AddUninitialized(BPP * NewOutput->Width * NewOutput->Height);
+
+			OutBufferData = NewOutput->Buffer->GetData();
+			OutBufferSize = NewOutput->Buffer->Num();
+		}
+
+		Result = Decoder.Decode(OutBufferData, OutBufferSize, InInputAccessUnit.Data, (unsigned int)InInputAccessUnit.DataSize);
+
+#if ELECTRA_MEDIAGPUBUFFER_DX12
+		if (bUseGPUBuffers)
+		{
+			if (Result == 0)
+			{
+				// If we have a temp buffer, we need to copy the data first...
+				if (!TempBuffer.IsEmpty())
+				{
+					if (!PinnedResourceDelegate->RunCodeAsync([BufferResource=NewOutput->GPUBuffer.Resource, Width=NewOutput->Width, Height=NewOutput->Height, ResourcePitch, BPP, TempBuffer, BufferFence=NewOutput->GPUBuffer.Fence, BufferFenceValue=NewOutput->GPUBuffer.FenceValue]()
+						{
+							CopyData(BufferResource, ResourcePitch, TempBuffer.GetData(), Width * BPP, Height);
+							BufferFence->Signal(BufferFenceValue);
+						}, TaskSync.Get()))
+					{
+						// Async copy failed, do it synchronously...
+						CopyData(NewOutput->GPUBuffer.Resource, ResourcePitch, TempBuffer.GetData(), NewOutput->Width * BPP, NewOutput->Height);
+						NewOutput->GPUBuffer.Fence->Signal(NewOutput->GPUBuffer.FenceValue);
+					}
+				}
+				else
+				{
+					// We decoded into the resource: Unmap the resource memory and signal that it's usable by the GPU
+					NewOutput->GPUBuffer.Resource->Unmap(0, nullptr);
+					NewOutput->GPUBuffer.Fence->Signal(NewOutput->GPUBuffer.FenceValue);
+				}
+			}
+		}
+#endif
 
-		Result = Decoder.Decode(NewOutput->Buffer, InInputAccessUnit.Data, (unsigned int)InInputAccessUnit.DataSize);
 		if (Result)
 		{
 			PostError(Result, TEXT("DNX_DecodeFrame() failed"), ERRCODE_INTERNAL_FAILED_TO_DECODE_INPUT);
 			return IElectraDecoder::EDecoderError::Error;
 		}
+
 		CurrentOutput = MoveTemp(NewOutput);
 	}
 	return IElectraDecoder::EDecoderError::None;
@@ -731,16 +892,11 @@ int32 FVideoDecoderAvidDNxHDElectra::FDecoderHandle::PrepareForDecoding(const vo
 			if (Result == DNX_NO_ERROR)
 			{
 				NumOutputBits = GetNumberOfOutputColorBitsPerPixel(CurrentUncompressedParams.compType);
-				BufferAllocSize = CurrentCompressedParams.width
-					* CurrentCompressedParams.height
-					* GetNumberOfOutputColorComponents(CurrentUncompressedParams.compOrder)
-					* GetNumberOfOutputColorBytesPerPixel(CurrentUncompressedParams.compType);
 			}
 			else
 			{
 				DNX_DeleteDecoder(Handle);
 				Handle = nullptr;
-				BufferAllocSize = 0;
 				NumOutputBits = 0;
 				bHasAlpha = false;
 			}
@@ -772,10 +928,6 @@ int32 FVideoDecoderAvidDNxHDElectra::FDecoderHandle::PrepareForDecoding(const vo
 		}
 
 		NumOutputBits = GetNumberOfOutputColorBitsPerPixel(CurrentUncompressedParams.compType);
-		BufferAllocSize = CurrentCompressedParams.width
-			* CurrentCompressedParams.height
-			* GetNumberOfOutputColorComponents(CurrentUncompressedParams.compOrder)
-			* GetNumberOfOutputColorBytesPerPixel(CurrentUncompressedParams.compType);
 	}
 
 	return Result;
@@ -828,13 +980,13 @@ bool FVideoDecoderAvidDNxHDElectra::FDecoderHandle::SetupDecompressStruct()
 	return true;
 }
 
-int32 FVideoDecoderAvidDNxHDElectra::FDecoderHandle::Decode(TSharedPtr<TArray<uint8>, ESPMode::ThreadSafe> OutBuffer, const void* InInput, uint32 InInputSize)
+int32 FVideoDecoderAvidDNxHDElectra::FDecoderHandle::Decode(uint8* OutBufferData, uint32 OutBufferSize, const void* InInput, uint32 InInputSize)
 {
 	if (!Handle)
 	{
 		return DNX_NOT_INITIALIZED_ERROR;
 	}
-	int32 Result = DNX_DecodeFrame(Handle, InInput, OutBuffer->GetData(), InInputSize, (unsigned int)OutBuffer->Num());
+	int32 Result = DNX_DecodeFrame(Handle, InInput, OutBufferData, InInputSize, (unsigned int)OutBufferSize);
 	return Result;
 }
 
@@ -846,59 +998,9 @@ void FVideoDecoderAvidDNxHDElectra::FDecoderHandle::Reset()
 		DNX_DeleteDecoder(Handle);
 		Handle = nullptr;
 	}
-	BufferAllocSize = 0;
 	NumOutputBits = 0;
 }
 
-int32 FVideoDecoderAvidDNxHDElectra::FDecoderHandle::GetNumberOfOutputColorComponents(DNX_ColorComponentOrder_t InOutputComponentOrder)
-{
-	switch(InOutputComponentOrder)
-	{
-		case DNX_ColorComponentOrder_t::DNX_CCO_CbYCrY_NoA:			// CbY0CrY1
-		case DNX_ColorComponentOrder_t::DNX_CCO_YCbYCr_NoA:			// Y0CbY1Cr
-		case DNX_ColorComponentOrder_t::DNX_CCO_RGB_NoA:			// RGB
-		case DNX_ColorComponentOrder_t::DNX_CCO_BGR_NoA:			// BGR
-		case DNX_ColorComponentOrder_t::DNX_CCO_YCbCr_Interleaved:	// YCbCr 444
-		case DNX_ColorComponentOrder_t::DNX_CCO_Ch1Ch2Ch3:			// Arbitrary 444 subsampled color components for any colorspace other than RGB and YCbCr
-		case DNX_ColorComponentOrder_t::DNX_CCO_YCbCr_Planar:		// YCbCr 420 stored in planar form. Y followed by Cb followed by Cr
-			return 3;
-		case DNX_ColorComponentOrder_t::DNX_CCO_ARGB_Interleaved:	// ARGB
-		case DNX_ColorComponentOrder_t::DNX_CCO_BGRA_Interleaved:	// BGRA
-		case DNX_ColorComponentOrder_t::DNX_CCO_RGBA_Interleaved:	// RGBA
-		case DNX_ColorComponentOrder_t::DNX_CCO_ABGR_Interleaved:	// ABGR
-		case DNX_ColorComponentOrder_t::DNX_CCO_Ch1Ch2Ch1Ch3:		// Arbitrary 422 subsampled color components for any colorspace other than RGB and YCbCr with not subsampled channel first
-		case DNX_ColorComponentOrder_t::DNX_CCO_Ch2Ch1Ch3Ch1:		// Arbitrary 422 subsampled color components for any colorspace other than RGB and YCbCr with subsampled channel first
-		case DNX_ColorComponentOrder_t::DNX_CCO_CbYACrYA_Interleaved: // YCbCrA 4224
-		case DNX_ColorComponentOrder_t::DNX_CCO_CbYCrA_Interleaved:	// YCbCrA 4444
-		case DNX_ColorComponentOrder_t::DNX_CCO_YCbCrA_Planar:		// YCbCrA 4204 stored in planar form. Y followed by Cb followed by Cr followed by A
-		case DNX_ColorComponentOrder_t::DNX_CCO_Ch1Ch2Ch3A:			// Arbitrary 4444 subsampled color components for any colorspace other than RGB and YCbCr with Alpha (Alpha channel last)
-		case DNX_ColorComponentOrder_t::DNX_CCO_Ch3Ch2Ch1A:			// Arbitrary 4444 subsampled color components for any colorspace other than RGB and YCbCr with Alpha with reversed channel order (Alpha channel last)
-		case DNX_ColorComponentOrder_t::DNX_CCO_ACh1Ch2Ch3:			// Arbitrary 4444 subsampled color components for any colorspace other than RGB and YCbCr with Alpha (Alpha channel first)
-		case DNX_ColorComponentOrder_t::DNX_CCO_Ch2Ch1ACh3Ch1A:		// Arbitrary 4224 subsampled color components for any colorspace other than RGB and YCbCr with not subsampled channel first with Alpha
-		case DNX_ColorComponentOrder_t::DNX_CCO_CbYCrY_A:			// YCbCrA 4224 mixed-planar (separate Alpha plane)
-			return 4;
-		default:
-			return 4;
-	}
-}
-
-int32 FVideoDecoderAvidDNxHDElectra::FDecoderHandle::GetNumberOfOutputColorBytesPerPixel(DNX_ComponentType_t InOutputComponentType)
-{
-	switch(InOutputComponentType)
-	{
-		default:
-		case DNX_ComponentType_t::DNX_CT_UCHAR:			// 8 bit
-			return 1;
-		case DNX_ComponentType_t::DNX_CT_USHORT_10_6:	// 10 bit
-		case DNX_ComponentType_t::DNX_CT_SHORT_2_14:	// Fixed point 
-		case DNX_ComponentType_t::DNX_CT_SHORT:			// 16 bit. Premultiplied by 257. Byte ordering is machine dependent.
-		case DNX_ComponentType_t::DNX_CT_10Bit_2_8:		// 10 bit in 2_8 format. Byte ordering is fixed. This is to be used with 10-bit 4:2:2 YCbCr components.
-		case DNX_ComponentType_t::DNX_CT_V210:			// Apple's V210 
-		case DNX_ComponentType_t::DNX_CT_USHORT_12_4:	// 12 bit
-			return 2;
-	}
-}
-
 int32 FVideoDecoderAvidDNxHDElectra::FDecoderHandle::GetNumberOfOutputColorBitsPerPixel(DNX_ComponentType_t InOutputComponentType)
 {
 	switch (InOutputComponentType)
diff --git a/Engine/Plugins/Media/ElectraCodecs/Source/ElectraDecoders/ElectraDecoders.Build.cs b/Engine/Plugins/Media/ElectraCodecs/Source/ElectraDecoders/ElectraDecoders.Build.cs
--- a/Engine/Plugins/Media/ElectraCodecs/Source/ElectraDecoders/ElectraDecoders.Build.cs
+++ b/Engine/Plugins/Media/ElectraCodecs/Source/ElectraDecoders/ElectraDecoders.Build.cs
@@ -44,6 +44,7 @@ namespace UnrealBuildTool.Rules
 			if (Target.Platform.IsInGroup(UnrealPlatformGroup.Windows))
 			{
 				PublicDependencyModuleNames.Add("DirectX");
+				AddEngineThirdPartyPrivateStaticDependencies(Target, "DX12");
 
 				PrivateDefinitions.Add("_CRT_SECURE_NO_WARNINGS=1");
 
diff --git a/Engine/Plugins/Media/ElectraCodecs/Source/ElectraDecoders/Private/Apple/h264/VideoDecoderH264_Apple.cpp b/Engine/Plugins/Media/ElectraCodecs/Source/ElectraDecoders/Private/Apple/h264/VideoDecoderH264_Apple.cpp
--- a/Engine/Plugins/Media/ElectraCodecs/Source/ElectraDecoders/Private/Apple/h264/VideoDecoderH264_Apple.cpp
+++ b/Engine/Plugins/Media/ElectraCodecs/Source/ElectraDecoders/Private/Apple/h264/VideoDecoderH264_Apple.cpp
@@ -96,6 +96,10 @@ public:
 	{
 		return InBufferIndex == 0 ? ImageBuffer : nullptr;
 	}
+	virtual bool GetBufferTextureSyncByIndex(int32 InBufferIndex, FElectraDecoderOutputSync& SyncObject) const override
+	{
+		return false;
+	}
 	EElectraDecoderPlatformPixelFormat GetBufferFormatByIndex(int32 InBufferIndex) const override
 	{
 		return PixelFormat;
diff --git a/Engine/Plugins/Media/ElectraCodecs/Source/ElectraDecoders/Private/Apple/h265/VideoDecoderH265_Apple.cpp b/Engine/Plugins/Media/ElectraCodecs/Source/ElectraDecoders/Private/Apple/h265/VideoDecoderH265_Apple.cpp
--- a/Engine/Plugins/Media/ElectraCodecs/Source/ElectraDecoders/Private/Apple/h265/VideoDecoderH265_Apple.cpp
+++ b/Engine/Plugins/Media/ElectraCodecs/Source/ElectraDecoders/Private/Apple/h265/VideoDecoderH265_Apple.cpp
@@ -96,6 +96,10 @@ public:
 	{
 		return InBufferIndex == 0 ? ImageBuffer : nullptr;
 	}
+	virtual bool GetBufferTextureSyncByIndex(int32 InBufferIndex, FElectraDecoderOutputSync& SyncObject) const override
+	{
+		return false;
+	}
 	EElectraDecoderPlatformPixelFormat GetBufferFormatByIndex(int32 InBufferIndex) const override
 	{
 		return PixelFormat;
diff --git a/Engine/Plugins/Media/ElectraCodecs/Source/ElectraDecoders/Private/DX/VideoDecoderH264_DX.cpp b/Engine/Plugins/Media/ElectraCodecs/Source/ElectraDecoders/Private/DX/VideoDecoderH264_DX.cpp
--- a/Engine/Plugins/Media/ElectraCodecs/Source/ElectraDecoders/Private/DX/VideoDecoderH264_DX.cpp
+++ b/Engine/Plugins/Media/ElectraCodecs/Source/ElectraDecoders/Private/DX/VideoDecoderH264_DX.cpp
@@ -13,6 +13,12 @@
 #include "IElectraDecoderResourceDelegate.h"
 #include "ElectraDecodersModule.h"
 
+#include COMPILED_PLATFORM_HEADER(ElectraDecoderGPUBufferHelpers.h)
+#include COMPILED_PLATFORM_HEADER(PlatformHeaders_Video_DX.h)
+
+#include "Stats/Stats.h"
+#include "HAL/IConsoleManager.h"
+
 /*********************************************************************************************************************/
 #include COMPILED_PLATFORM_HEADER(PlatformHeaders_Video_DX.h)
 
@@ -24,6 +30,25 @@ if (FAILED(Result))						\
 	return false;						\
 }
 
+/*********************************************************************************************************************/
+
+#ifndef ELECTRA_HAVE_IMAGEBUFFERS
+#define	ELECTRA_HAVE_IMAGEBUFFERS 0
+#endif
+
+/*********************************************************************************************************************/
+
+DECLARE_CYCLE_STAT(TEXT("ElectraDecoder ConvertOutput"), STAT_ElectraDecoder_ConvertOutputH264, STATGROUP_Media);
+
+/*********************************************************************************************************************/
+
+static TAutoConsoleVariable<int32> CVarElectraWindowsH264UseOldOutputPath(
+	TEXT("Electra.Win.H264UseOldOutputPath"),
+	0,
+	TEXT("Use old style CPU buffer output path for H264.\n")
+	TEXT(" 0: use new upload heap / direct path (default); 1: use old CPU buffer path."),
+	ECVF_Default);
+
 /*********************************************************************************************************************/
 /*********************************************************************************************************************/
 /*********************************************************************************************************************/
@@ -38,11 +63,13 @@ public:
 
 
 class FElectraVideoDecoderOutputH264_DX : public IElectraDecoderVideoOutput
+#if ELECTRA_HAVE_IMAGEBUFFERS
+										, public IElectraDecoderVideoOutputImageBuffers
+#endif
 {
 public:
 	virtual ~FElectraVideoDecoderOutputH264_DX()
-	{
-	}
+	{ }
 
 	FTimespan GetPTS() const override
 	{ return PTS; }
@@ -75,6 +102,17 @@ public:
 	{ 
 		switch(InTypeOfHandle)
 		{
+#if ELECTRA_HAVE_IMAGEBUFFERS
+			case EElectraDecoderPlatformOutputHandleType::ImageBuffers:
+				if (!MFSample.IsValid())
+				{
+					if (InTypeOfHandle == EElectraDecoderPlatformOutputHandleType::ImageBuffers)
+					{
+						return static_cast<IElectraDecoderVideoOutputImageBuffers*>(const_cast<FElectraVideoDecoderOutputH264_DX*>(this));
+					}
+				}
+				return nullptr;
+#endif
 			case EElectraDecoderPlatformOutputHandleType::MFSample:
 				return MFSample.GetReference(); 
 			case EElectraDecoderPlatformOutputHandleType::DXDevice:
@@ -90,6 +128,86 @@ public:
 	IElectraDecoderVideoOutput::EImageCopyResult CopyPlatformImage(IElectraDecoderVideoOutputCopyResources* InCopyResources) const override
 	{ return IElectraDecoderVideoOutput::EImageCopyResult::NotSupported; }
 
+#if ELECTRA_HAVE_IMAGEBUFFERS
+	//
+	// IElectraDecoderVideoOutputImageBuffers
+	//
+
+	// Return the 4cc of the codec.
+	virtual uint32 GetCodec4CC() const override
+	{
+		return *(uint32*)"h264";
+	}
+
+	// Returns the number of separate image buffers making up the frame.
+	virtual int32 GetNumberOfBuffers() const override
+	{
+		return 1;
+	}
+
+	// Returns the n'th image buffer as CPU data buffer.
+	virtual TSharedPtr<TArray<uint8>, ESPMode::ThreadSafe> GetBufferDataByIndex(int32 InBufferIndex) const override
+	{
+		// For now CPU side buffers use the old style MFSample way to return data
+		return nullptr;
+	}
+
+	// Returns the n'th image buffer as GPU texture resource.
+	virtual void* GetBufferTextureByIndex(int32 InBufferIndex) const override
+	{
+		if (InBufferIndex == 0)
+		{
+			return Buffer.Resource.GetReference();
+		}
+		return nullptr;
+	}
+
+	virtual bool GetBufferTextureSyncByIndex(int32 InBufferIndex, FElectraDecoderOutputSync& SyncObject) const override
+	{
+		if (InBufferIndex == 0)
+		{
+#if ALLOW_MFSAMPLE_WITH_DX12
+			SyncObject = { Buffer.DecoderSync.GetReference(), 0, Buffer.MFSample };
+			// note: a reference to the MFSample remains with this instance, but is assumed to soon be released as the instance will be destroyed shortly after this was called
+#else
+			SyncObject = { Buffer.Fence.GetReference(), Buffer.FenceValue, nullptr };
+#endif
+			return true;
+		}
+		return false;
+	}
+
+	// Returns the n'th image buffer format
+	virtual EElectraDecoderPlatformPixelFormat GetBufferFormatByIndex(int32 InBufferIndex) const override
+	{
+		if (InBufferIndex == 0)
+		{
+			return (EElectraDecoderPlatformPixelFormat)ElectraDecodersUtil::GetVariantValueSafeU64(ExtraValues, TEXT("pixfmt"), (uint64)EElectraDecoderPlatformPixelFormat::INVALID);
+		}
+		return EElectraDecoderPlatformPixelFormat::INVALID;
+	}
+
+	// Returns the n'th image buffer encoding
+	virtual EElectraDecoderPlatformPixelEncoding GetBufferEncodingByIndex(int32 InBufferIndex) const override
+	{
+		if (InBufferIndex == 0)
+		{
+			return (EElectraDecoderPlatformPixelEncoding)ElectraDecodersUtil::GetVariantValueSafeU64(ExtraValues, TEXT("pixenc"), (uint64)EElectraDecoderPlatformPixelEncoding::Native);
+		}
+		return EElectraDecoderPlatformPixelEncoding::Native;
+	}
+	
+	// Returns the n'th image buffer pitch
+	virtual int32 GetBufferPitchByIndex(int32 InBufferIndex) const override
+	{
+		if (InBufferIndex == 0)
+		{
+			return Pitch;
+		}
+		return 0;
+	}
+#endif // ELECTRA_HAVE_IMAGEBUFFERS
+
 public:
 	FTimespan PTS;
 	uint64 UserValue = 0;
@@ -109,6 +227,16 @@ public:
 	TMap<FString, FVariant> ExtraValues;
 	TRefCountPtr<IMFSample> MFSample;
 	FElectraVideoDecoderDXDeviceContext Ctx;
+
+#if ALLOW_MFSAMPLE_WITH_DX12
+	struct {
+		TRefCountPtr<ID3D12Resource> Resource;
+		TRefCountPtr<IMFD3D12SynchronizationObjectCommands> DecoderSync;
+		TRefCountPtr<IMFSample> MFSample;	// note: this reference is kept to allow proper passing on of ref keeping to the DX12 queue only - as this instance will be destroyed once the DX12 queue took over, this reference will not keep the MFSample around for too long
+	} Buffer;
+#else
+	FElectraMediaDecoderOutputBufferPool_DX12::FOutputData Buffer;
+#endif
 };
 
 
@@ -219,6 +347,7 @@ private:
 
 	bool CreateInputSample(TRefCountPtr<IMFSample>& InputSample, const FInputAccessUnit& InInputAccessUnit, const TMap<FString, FVariant>& InAdditionalOptions);
 	bool ConvertDecoderOutput();
+	static void CopyBufferData(TRefCountPtr<IMF2DBuffer> Buffer2D, TRefCountPtr<ID3D12Resource> BufferResource, TRefCountPtr<ID3D12Fence> BufferFence, uint64 BufferFenceValue, uint32 BufferPitch, int32 DecodedHeight);
 
 
 	TWeakPtr<IElectraDecoderResourceDelegate, ESPMode::ThreadSafe> ResourceDelegate;
@@ -236,6 +365,14 @@ private:
 	TRefCountPtr<IMFTransform> DecoderTransform;
 	TRefCountPtr<IMFMediaType> CurrentOutputMediaType;
 
+	uint32 MaxOutputBuffers;
+	uint32 MaxWidth;
+	uint32 MaxHeight;
+#if !ALLOW_MFSAMPLE_WITH_DX12
+	mutable TSharedPtr<FElectraMediaDecoderOutputBufferPool_DX12> D3D12ResourcePool;
+	mutable TUniquePtr<IElectraDecoderResourceDelegateBase::IAsyncConsecutiveTaskSync> TaskSync;
+#endif
+
 	static const GUID MFTms_AVLowLatencyMode;
 };
 
@@ -257,6 +394,12 @@ TSharedPtr<IElectraDecoder, ESPMode::ThreadSafe> IElectraVideoDecoderH264_DX::Cr
 FElectraVideoDecoderH264_DX::FElectraVideoDecoderH264_DX(const TMap<FString, FVariant>& InOptions, TSharedPtr<IElectraDecoderResourceDelegate, ESPMode::ThreadSafe> InResourceDelegate)
 {
 	ResourceDelegate = InResourceDelegate;
+
+	MaxWidth = (uint32)ElectraDecodersUtil::GetVariantValueSafeU64(InOptions, TEXT("max_width"), 1920);
+	MaxHeight = (uint32)ElectraDecodersUtil::GetVariantValueSafeU64(InOptions, TEXT("max_height"), 1080);
+
+	MaxOutputBuffers = (uint32)ElectraDecodersUtil::GetVariantValueSafeU64(InOptions, TEXT("max_output_buffers"), 5);
+	MaxOutputBuffers += kElectraDecoderPipelineExtraFrames;
 }
 
 
@@ -456,7 +599,7 @@ IElectraDecoder::EDecoderError FElectraVideoDecoderH264_DX::DecodeAccessUnit(con
 				{
 					return IElectraDecoder::EDecoderError::Error;
 				}
-				// It is possible that we got output without having consumed the input. Check for this and return appropraitely.
+				// It is possible that we got output without having consumed the input. Check for this and return appropriately.
 				return InputSample.IsValid() ? IElectraDecoder::EDecoderError::NoBuffer : IElectraDecoder::EDecoderError::None;
 			}
 			else
@@ -807,9 +950,32 @@ bool FElectraVideoDecoderH264_DX::CreateInputSample(TRefCountPtr<IMFSample>& Inp
 	return true;
 }
 
+#if !ALLOW_MFSAMPLE_WITH_DX12
+void FElectraVideoDecoderH264_DX::CopyBufferData(TRefCountPtr<IMF2DBuffer> Buffer2D, TRefCountPtr<ID3D12Resource> BufferResource, TRefCountPtr<ID3D12Fence> BufferFence, uint64 BufferFenceValue, uint32 BufferPitch, int32 DecodedHeight)
+{
+	uint8* Data = nullptr;
+	LONG Pitch = 0;
+	HRESULT Res = Buffer2D->Lock2D(&Data, &Pitch);
+	check(SUCCEEDED(Res));
+
+	void* BufferAddr;
+	Res = BufferResource->Map(0, nullptr, &BufferAddr);
+	check(SUCCEEDED(Res));
+
+	FElectraMediaDecoderOutputBufferPool_DX12::CopyWithPitchAdjust((uint8*)BufferAddr, BufferPitch, Data, Pitch, DecodedHeight * 3 / 2);
+
+	BufferResource->Unmap(0, nullptr);
+
+	Buffer2D->Unlock2D();
+
+	BufferFence->Signal(BufferFenceValue);
+};
+#endif
 
 bool FElectraVideoDecoderH264_DX::ConvertDecoderOutput()
 {
+	SCOPE_CYCLE_COUNTER(STAT_ElectraDecoder_ConvertOutputH264);
+
 	TRefCountPtr<IMFSample> DecodedOutputSample = CurrentDecoderOutputBuffer->DetachOutputSample();
 	CurrentDecoderOutputBuffer.Reset();
 	if (!DecodedOutputSample.IsValid())
@@ -901,8 +1067,13 @@ bool FElectraVideoDecoderH264_DX::ConvertDecoderOutput()
 	// (DirectX12 and higher (unless we can use the render device to decode), SW decoder or a non-DirectX API will trigger this)
 	if (DecoderPlatformHandle->IsSoftware() || (DecoderPlatformHandle->GetDXVersionTimes1000() >= 12000 && DecoderPlatformHandle->GetDXDevice() != nullptr) || DecoderPlatformHandle->GetDXVersionTimes1000() == 0)
 	{
-		// We are returning a CPU side buffer (software decode OR DX12 render device) - adjust the height so it can be interpreted as a single plane texture
-		NewOutput->DecodedHeight = NewOutput->DecodedHeight * 3 / 2;
+#if !ALLOW_MFSAMPLE_WITH_DX12
+		if (DecoderPlatformHandle->GetDXVersionTimes1000() < 12000 || CVarElectraWindowsH264UseOldOutputPath.GetValueOnAnyThread() != 0)
+#endif
+		{
+			// We are returning a CPU side buffer (software decode OR DX12 render device) - adjust the height so it can be interpreted as a single plane texture
+			NewOutput->DecodedHeight = NewOutput->DecodedHeight * 3 / 2;
+		}
 	}
 
 	// Note: Decoder crops at the lower right border.
@@ -913,7 +1084,6 @@ bool FElectraVideoDecoderH264_DX::ConvertDecoderOutput()
 
 	// Try to get the stride. Defaults to 0 should it not be obtainable.
 	UINT32 stride = MFGetAttributeUINT32(CurrentOutputMediaType, MF_MT_DEFAULT_STRIDE, 0);
-	NewOutput->Pitch = stride;
 
 	// Try to get the frame rate ratio.
 	UINT32 num=0, denom=0;
@@ -940,8 +1110,113 @@ bool FElectraVideoDecoderH264_DX::ConvertDecoderOutput()
 	NewOutput->ExtraValues.Emplace(TEXT("pixfmt"), FVariant((int64)EElectraDecoderPlatformPixelFormat::NV12));
 	NewOutput->ExtraValues.Emplace(TEXT("pixenc"), FVariant((int64)EElectraDecoderPlatformPixelEncoding::Native));
 
-	// Retain the IMFSample in the output. It is needed later in converting it for display.
-	NewOutput->MFSample = DecodedOutputSample;
+#if ELECTRA_HAVE_IMAGEBUFFERS
+	void* PlatformDevice = nullptr;
+	int32 PlatformDeviceVersion = 0;
+	bool bUseGPUBuffers = false;
+	auto PinnedResourceDelegate = ResourceDelegate.Pin();
+	if (PinnedResourceDelegate.IsValid() && CVarElectraWindowsH264UseOldOutputPath.GetValueOnAnyThread() == 0)
+	{
+		PinnedResourceDelegate->GetD3DDevice(&PlatformDevice, &PlatformDeviceVersion);
+		bUseGPUBuffers = (PlatformDevice && PlatformDeviceVersion >= 12000);
+	}
+
+	// If we don't have support for DX12 direct to video decoding, yet, we bring the CPU buffer data into a DX12 resource now
+	if (bUseGPUBuffers)
+	{
+		DWORD BuffersNum = 0;
+		if ((Result = DecodedOutputSample->GetBufferCount(&BuffersNum)) != S_OK)
+		{
+			return false;
+		}
+		if (BuffersNum != 1)
+		{
+			return false;
+		}
+		TRefCountPtr<IMFMediaBuffer> Buffer;
+		if ((Result = DecodedOutputSample->GetBufferByIndex(0, Buffer.GetInitReference())) != S_OK)
+		{
+			return false;
+		}
+
+#if !ALLOW_MFSAMPLE_WITH_DX12
+		//
+		// Pre SDK 22621: use CPU copy to get data from MFSample to DX12 resource
+		//
+		TRefCountPtr D3D12Device(static_cast<ID3D12Device*>(PlatformDevice));
+
+		// Create the resource pool as needed...
+		if (!D3D12ResourcePool)
+		{
+			D3D12ResourcePool = MakeShared<FElectraMediaDecoderOutputBufferPool_DX12, ESPMode::ThreadSafe>(D3D12Device, MaxOutputBuffers, MaxWidth, MaxHeight * 3 / 2, 1);
+		}
+
+		// Create a tasksync instance so we can have our own async jobs run in consecutive order
+		if (!TaskSync.IsValid())
+		{
+			TaskSync.Reset(PinnedResourceDelegate->CreateAsyncConsecutiveTaskSync());
+		}
+
+		// Get data from MFSample...
+		TRefCountPtr<IMF2DBuffer> Buffer2D;
+		if ((Result = Buffer->QueryInterface(__uuidof(IMF2DBuffer), (void**)Buffer2D.GetInitReference())) != S_OK)
+		{
+			return false;
+		}
+
+		// Request resource and fence...
+		uint32 BufferPitch;
+		D3D12ResourcePool->AllocateOutputDataAsBuffer(NewOutput->Buffer, BufferPitch);
+
+		NewOutput->Pitch = BufferPitch;
+
+		if (!PinnedResourceDelegate->RunCodeAsync([DecodedOutputSample, Buffer2D, BufferResource = NewOutput->Buffer.Resource, BufferFence = NewOutput->Buffer.Fence, BufferFenceValue = NewOutput->Buffer.FenceValue, BufferPitch, DecodedHeight = NewOutput->DecodedHeight]()
+			{
+				CopyBufferData(Buffer2D, BufferResource, BufferFence, BufferFenceValue, BufferPitch, DecodedHeight);
+			}, TaskSync.Get()))
+		{
+			CopyBufferData(Buffer2D, NewOutput->Buffer.Resource, NewOutput->Buffer.Fence, NewOutput->Buffer.FenceValue, BufferPitch, NewOutput->DecodedHeight);
+		}
+#else
+		//
+		// Post SDK 22621: direct link from MFSample to DX12 resource
+		//
+		
+		TRefCountPtr<IMFDXGIBuffer> DXGIBuffer;
+		if ((Result = Buffer->QueryInterface(__uuidof(IMFDXGIBuffer), (void**)DXGIBuffer.GetInitReference())) != S_OK)
+		{
+			return false;
+		}
+
+		TRefCountPtr<ID3D12Resource> Resource;
+		Result = DXGIBuffer->GetResource(__uuidof(ID3D12Resource), (void**)Resource.GetInitReference());
+		if (Result == S_OK)
+		{
+			// Yes, set this up...
+			TRefCountPtr<IMFD3D12SynchronizationObjectCommands> DecoderSync;
+			if ((Result = DXGIBuffer->GetUnknown(MF_D3D12_SYNCHRONIZATION_OBJECT, IID_PPV_ARGS(DecoderSync.GetInitReference()))) != S_OK)
+			{
+				return false;
+			}
+
+			NewOutput->Buffer.Resource = Resource;
+			NewOutput->Buffer.DecoderSync = DecoderSync;
+			NewOutput->Buffer.MFSample = DecodedOutputSample;	// we keep a reference around until we actually submit the sync object to a DX12 queue
+		}
+#endif // !ALLOW_MFSAMPLE_WITH_DX12
+	}
+	else
+#endif // ELECTRA_HAVE_IMAGEBUFFERS
+	{
+		//
+		// Fallback output path for pre- / none-DX12 graphics APIs
+		// 
+		
+		// Retain the IMFSample in the output. It is needed later in converting it for display.
+		NewOutput->Pitch = stride;
+		NewOutput->MFSample = DecodedOutputSample;
+	}
+
 	// Do the same for the HW acceleration device.
 	NewOutput->Ctx.SetDeviceAndContext(DecoderPlatformHandle->GetDXDevice(), DecoderPlatformHandle->GetDXDeviceContext());
 
diff --git a/Engine/Plugins/Media/ElectraCodecs/Source/ElectraDecoders/Private/DX/VideoDecoderH265_DX.cpp b/Engine/Plugins/Media/ElectraCodecs/Source/ElectraDecoders/Private/DX/VideoDecoderH265_DX.cpp
--- a/Engine/Plugins/Media/ElectraCodecs/Source/ElectraDecoders/Private/DX/VideoDecoderH265_DX.cpp
+++ b/Engine/Plugins/Media/ElectraCodecs/Source/ElectraDecoders/Private/DX/VideoDecoderH265_DX.cpp
@@ -13,6 +13,11 @@
 #include "IElectraDecoderResourceDelegate.h"
 #include "ElectraDecodersModule.h"
 
+#include COMPILED_PLATFORM_HEADER(ElectraDecoderGPUBufferHelpers.h)
+#include COMPILED_PLATFORM_HEADER(PlatformHeaders_Video_DX.h)
+
+#include "Stats\Stats.h"
+
 /*********************************************************************************************************************/
 #include COMPILED_PLATFORM_HEADER(PlatformHeaders_Video_DX.h)
 
@@ -24,6 +29,16 @@ if (FAILED(Result))						\
 	return false;						\
 }
 
+/*********************************************************************************************************************/
+
+#ifndef ELECTRA_HAVE_IMAGEBUFFERS
+#define	ELECTRA_HAVE_IMAGEBUFFERS 0
+#endif
+
+/*********************************************************************************************************************/
+
+DECLARE_CYCLE_STAT(TEXT("ElectraDecoder ConvertOutput"), STAT_ElectraDecoder_ConvertOutputH265, STATGROUP_Media);
+
 /*********************************************************************************************************************/
 /*********************************************************************************************************************/
 /*********************************************************************************************************************/
@@ -38,57 +53,180 @@ public:
 
 
 class FElectraVideoDecoderOutputH265_DX : public IElectraDecoderVideoOutput
+#if ELECTRA_HAVE_IMAGEBUFFERS
+	, public IElectraDecoderVideoOutputImageBuffers
+#endif
 {
 public:
 	virtual ~FElectraVideoDecoderOutputH265_DX()
-	{
-	}
+	{ }
 
 	FTimespan GetPTS() const override
-	{ return PTS; }
+	{
+		return PTS;
+	}
 	uint64 GetUserValue() const override
-	{ return UserValue; }
+	{
+		return UserValue;
+	}
 
 	int32 GetWidth() const override
-	{ return Width; }
+	{
+		return Width;
+	}
 	int32 GetHeight() const override
-	{ return Height; }
+	{
+		return Height;
+	}
 	int32 GetDecodedWidth() const override
-	{ return DecodedWidth; }
+	{
+		return DecodedWidth;
+	}
 	int32 GetDecodedHeight() const override
-	{ return DecodedHeight; }
+	{
+		return DecodedHeight;
+	}
 	FElectraVideoDecoderOutputCropValues GetCropValues() const override
-	{ return Crop; }
+	{
+		return Crop;
+	}
 	int32 GetAspectRatioW() const override
-	{ return AspectW; }
+	{
+		return AspectW;
+	}
 	int32 GetAspectRatioH() const override
-	{ return AspectH; }
+	{
+		return AspectH;
+	}
 	int32 GetFrameRateNumerator() const override
-	{ return FrameRateN; }
+	{
+		return FrameRateN;
+	}
 	int32 GetFrameRateDenominator() const override
-	{ return FrameRateD; }
+	{
+		return FrameRateD;
+	}
 	int32 GetNumberOfBits() const override
-	{ return NumBits; }
+	{
+		return NumBits;
+	}
 	void GetExtraValues(TMap<FString, FVariant>& OutExtraValues) const override
-	{ OutExtraValues = ExtraValues; }
+	{
+		OutExtraValues = ExtraValues;
+	}
 	void* GetPlatformOutputHandle(EElectraDecoderPlatformOutputHandleType InTypeOfHandle) const override
-	{ 
-		switch(InTypeOfHandle)
+	{
+		switch (InTypeOfHandle)
 		{
-			case EElectraDecoderPlatformOutputHandleType::MFSample:
-				return MFSample.GetReference(); 
-			case EElectraDecoderPlatformOutputHandleType::DXDevice:
-				return Ctx.DxDevice.GetReference(); 
-			case EElectraDecoderPlatformOutputHandleType::DXDeviceContext:
-				return Ctx.DxDeviceContext.GetReference(); 
-			default:
-				return nullptr;
+#if ELECTRA_HAVE_IMAGEBUFFERS
+		case EElectraDecoderPlatformOutputHandleType::ImageBuffers:
+			if (!MFSample.IsValid())
+			{
+				if (InTypeOfHandle == EElectraDecoderPlatformOutputHandleType::ImageBuffers)
+				{
+					return static_cast<IElectraDecoderVideoOutputImageBuffers*>(const_cast<FElectraVideoDecoderOutputH265_DX*>(this));
+				}
+			}
+			return nullptr;
+#endif
+		case EElectraDecoderPlatformOutputHandleType::MFSample:
+			return MFSample.GetReference();
+		case EElectraDecoderPlatformOutputHandleType::DXDevice:
+			return Ctx.DxDevice.GetReference();
+		case EElectraDecoderPlatformOutputHandleType::DXDeviceContext:
+			return Ctx.DxDeviceContext.GetReference();
+		default:
+			return nullptr;
 		}
 	}
 	IElectraDecoderVideoOutputTransferHandle* GetTransferHandle() const override
-	{ return nullptr; }
+	{
+		return nullptr;
+	}
 	IElectraDecoderVideoOutput::EImageCopyResult CopyPlatformImage(IElectraDecoderVideoOutputCopyResources* InCopyResources) const override
-	{ return IElectraDecoderVideoOutput::EImageCopyResult::NotSupported; }
+	{
+		return IElectraDecoderVideoOutput::EImageCopyResult::NotSupported;
+	}
+
+#if ELECTRA_HAVE_IMAGEBUFFERS
+	//
+	// IElectraDecoderVideoOutputImageBuffers
+	//
+
+	// Return the 4cc of the codec.
+	virtual uint32 GetCodec4CC() const override
+	{
+		return *(uint32*)"hevc";
+	}
+
+	// Returns the number of separate image buffers making up the frame.
+	virtual int32 GetNumberOfBuffers() const override
+	{
+		return 1;
+	}
+
+	// Returns the n'th image buffer as CPU data buffer.
+	virtual TSharedPtr<TArray<uint8>, ESPMode::ThreadSafe> GetBufferDataByIndex(int32 InBufferIndex) const override
+	{
+		// For now CPU side buffers use the old style MFSample way to return data
+		return nullptr;
+	}
+
+	// Returns the n'th image buffer as GPU texture resource.
+	virtual void* GetBufferTextureByIndex(int32 InBufferIndex) const override
+	{
+		if (InBufferIndex == 0)
+		{
+			return Buffer.Resource.GetReference();
+		}
+		return nullptr;
+	}
+
+	virtual bool GetBufferTextureSyncByIndex(int32 InBufferIndex, FElectraDecoderOutputSync& SyncObject) const override
+	{
+		if (InBufferIndex == 0)
+		{
+#if ALLOW_MFSAMPLE_WITH_DX12
+			SyncObject = { Buffer.DecoderSync.GetReference(), 0, Buffer.MFSample };
+			// note: a reference to the MFSample remains with this instance, but is assumed to soon be released as the instance will be destroyed shortly after this was called
+#else
+			SyncObject = { Buffer.Fence.GetReference(), Buffer.FenceValue, nullptr };
+#endif
+			return true;
+		}
+		return false;
+	}
+
+	// Returns the n'th image buffer format
+	virtual EElectraDecoderPlatformPixelFormat GetBufferFormatByIndex(int32 InBufferIndex) const override
+	{
+		if (InBufferIndex == 0)
+		{
+			return (EElectraDecoderPlatformPixelFormat)ElectraDecodersUtil::GetVariantValueSafeU64(ExtraValues, TEXT("pixfmt"), (uint64)EElectraDecoderPlatformPixelFormat::INVALID);
+		}
+		return EElectraDecoderPlatformPixelFormat::INVALID;
+	}
+
+	// Returns the n'th image buffer encoding
+	virtual EElectraDecoderPlatformPixelEncoding GetBufferEncodingByIndex(int32 InBufferIndex) const override
+	{
+		if (InBufferIndex == 0)
+		{
+			return (EElectraDecoderPlatformPixelEncoding)ElectraDecodersUtil::GetVariantValueSafeU64(ExtraValues, TEXT("pixenc"), (uint64)EElectraDecoderPlatformPixelEncoding::Native);
+		}
+		return EElectraDecoderPlatformPixelEncoding::Native;
+	}
+	
+	// Returns the n'th image buffer pitch
+	virtual int32 GetBufferPitchByIndex(int32 InBufferIndex) const override
+	{
+		if (InBufferIndex == 0)
+		{
+			return Pitch;
+		}
+		return 0;
+	}
+#endif // ELECTRA_HAVE_IMAGEBUFFERS
 
 public:
 	FTimespan PTS;
@@ -109,6 +247,16 @@ public:
 	TMap<FString, FVariant> ExtraValues;
 	TRefCountPtr<IMFSample> MFSample;
 	FElectraVideoDecoderDXDeviceContext Ctx;
+
+#if ALLOW_MFSAMPLE_WITH_DX12
+	struct {
+		TRefCountPtr<ID3D12Resource> Resource;
+		TRefCountPtr<IMFD3D12SynchronizationObjectCommands> DecoderSync;
+		TRefCountPtr<IMFSample> MFSample;	// note: this reference is kept to allow proper passing on of ref keeping to the DX12 queue only - as this instance will be destroyed once the DX12 queue took over, this reference will not keep the MFSample around for too long
+	} Buffer;
+#else
+	FElectraMediaDecoderOutputBufferPool_DX12::FOutputData Buffer;
+#endif
 };
 
 
@@ -219,7 +367,7 @@ private:
 
 	bool CreateInputSample(TRefCountPtr<IMFSample>& InputSample, const FInputAccessUnit& InInputAccessUnit, const TMap<FString, FVariant>& InAdditionalOptions);
 	bool ConvertDecoderOutput();
-
+	static void CopyBufferData(TRefCountPtr<IMF2DBuffer> Buffer2D, TRefCountPtr<ID3D12Resource> BufferResource, TRefCountPtr<ID3D12Fence> BufferFence, uint64 BufferFenceValue, uint32 BufferPitch, int32 DecodedHeight);
 
 	TWeakPtr<IElectraDecoderResourceDelegate, ESPMode::ThreadSafe> ResourceDelegate;
 
@@ -236,6 +384,14 @@ private:
 	TRefCountPtr<IMFTransform> DecoderTransform;
 	TRefCountPtr<IMFMediaType> CurrentOutputMediaType;
 
+	uint32 MaxOutputBuffers;
+	uint32 MaxWidth;
+	uint32 MaxHeight;
+#if !ALLOW_MFSAMPLE_WITH_DX12
+	mutable TSharedPtr<FElectraMediaDecoderOutputBufferPool_DX12> D3D12ResourcePool;
+	mutable TUniquePtr<IElectraDecoderResourceDelegateBase::IAsyncConsecutiveTaskSync> TaskSync;
+#endif
+
 	static const GUID MFTms_AVLowLatencyMode;
 };
 
@@ -257,6 +413,12 @@ TSharedPtr<IElectraDecoder, ESPMode::ThreadSafe> IElectraVideoDecoderH265_DX::Cr
 FElectraVideoDecoderH265_DX::FElectraVideoDecoderH265_DX(const TMap<FString, FVariant>& InOptions, TSharedPtr<IElectraDecoderResourceDelegate, ESPMode::ThreadSafe> InResourceDelegate)
 {
 	ResourceDelegate = InResourceDelegate;
+
+	MaxWidth = (uint32)ElectraDecodersUtil::GetVariantValueSafeU64(InOptions, TEXT("max_width"), 1920);
+	MaxHeight = (uint32)ElectraDecodersUtil::GetVariantValueSafeU64(InOptions, TEXT("max_height"), 1080);
+
+	MaxOutputBuffers = (uint32)ElectraDecodersUtil::GetVariantValueSafeU64(InOptions, TEXT("max_output_buffers"), 5);
+	MaxOutputBuffers += kElectraDecoderPipelineExtraFrames;
 }
 
 
@@ -457,7 +619,7 @@ IElectraDecoder::EDecoderError FElectraVideoDecoderH265_DX::DecodeAccessUnit(con
 				{
 					return IElectraDecoder::EDecoderError::Error;
 				}
-				// It is possible that we got output without having consumed the input. Check for this and return appropraitely.
+				// It is possible that we got output without having consumed the input. Check for this and return appropriately.
 				return InputSample.IsValid() ? IElectraDecoder::EDecoderError::NoBuffer : IElectraDecoder::EDecoderError::None;
 			}
 			else
@@ -822,9 +984,32 @@ bool FElectraVideoDecoderH265_DX::CreateInputSample(TRefCountPtr<IMFSample>& Inp
 	return true;
 }
 
+#if !ALLOW_MFSAMPLE_WITH_DX12
+void FElectraVideoDecoderH265_DX::CopyBufferData(TRefCountPtr<IMF2DBuffer> Buffer2D, TRefCountPtr<ID3D12Resource> BufferResource, TRefCountPtr<ID3D12Fence> BufferFence, uint64 BufferFenceValue, uint32 BufferPitch, int32 DecodedHeight)
+{
+	uint8* Data = nullptr;
+	LONG Pitch = 0;
+	HRESULT Res = Buffer2D->Lock2D(&Data, &Pitch);
+	check(SUCCEEDED(Res));
+
+	void* BufferAddr;
+	Res = BufferResource->Map(0, nullptr, &BufferAddr);
+	check(SUCCEEDED(Res));
+
+	FElectraMediaDecoderOutputBufferPool_DX12::CopyWithPitchAdjust((uint8*)BufferAddr, BufferPitch, Data, Pitch, DecodedHeight * 3 / 2);
+
+	BufferResource->Unmap(0, nullptr);
+
+	Buffer2D->Unlock2D();
+
+	BufferFence->Signal(BufferFenceValue);
+};
+#endif
 
 bool FElectraVideoDecoderH265_DX::ConvertDecoderOutput()
 {
+	SCOPE_CYCLE_COUNTER(STAT_ElectraDecoder_ConvertOutputH265);
+
 	TRefCountPtr<IMFSample> DecodedOutputSample = CurrentDecoderOutputBuffer->DetachOutputSample();
 	CurrentDecoderOutputBuffer.Reset();
 	if (!DecodedOutputSample.IsValid())
@@ -915,8 +1100,13 @@ bool FElectraVideoDecoderH265_DX::ConvertDecoderOutput()
 	// (DirectX12 and higher (unless we can use the render device to decode), SW decoder or a non-DirectX API will trigger this)
 	if (DecoderPlatformHandle->IsSoftware() || (DecoderPlatformHandle->GetDXVersionTimes1000() >= 12000 && DecoderPlatformHandle->GetDXDevice() != nullptr) || DecoderPlatformHandle->GetDXVersionTimes1000() == 0)
 	{
-		// We are returning a CPU side buffer (software decode OR DX12 render device) - adjust the height so it can be interpreted as a single plane texture
-		NewOutput->DecodedHeight = NewOutput->DecodedHeight * 3 / 2;
+#if !ALLOW_MFSAMPLE_WITH_DX12
+		if (DecoderPlatformHandle->GetDXVersionTimes1000() < 12000)
+#endif
+		{
+			// We are returning a CPU side buffer (software decode OR DX12 render device) - adjust the height so it can be interpreted as a single plane texture
+			NewOutput->DecodedHeight = NewOutput->DecodedHeight * 3 / 2;
+		}
 	}
 
 	// Note: Decoder crops at the lower right border.
@@ -927,20 +1117,19 @@ bool FElectraVideoDecoderH265_DX::ConvertDecoderOutput()
 
 	// Try to get the stride. Defaults to 0 should it not be obtainable.
 	UINT32 stride = MFGetAttributeUINT32(CurrentOutputMediaType, MF_MT_DEFAULT_STRIDE, 0);
-	NewOutput->Pitch = stride;
 
 	// Try to get the frame rate ratio.
-	UINT32 num=0, denom=0;
+	UINT32 num = 0, denom = 0;
 	MFGetAttributeRatio(CurrentOutputMediaType, MF_MT_FRAME_RATE, &num, &denom);
 	NewOutput->FrameRateN = num;
 	NewOutput->FrameRateD = denom;
 
 	// Try to get the pixel aspect ratio.
-	num=0, denom=0;
+	num = 0, denom = 0;
 	MFGetAttributeRatio(CurrentOutputMediaType, MF_MT_PIXEL_ASPECT_RATIO, &num, &denom);
 	if (!num || !denom)
 	{
-		num   = 1;
+		num = 1;
 		denom = 1;
 	}
 	NewOutput->AspectW = num;
@@ -970,14 +1159,119 @@ bool FElectraVideoDecoderH265_DX::ConvertDecoderOutput()
 		PixFmt = EElectraDecoderPlatformPixelFormat::P010;
 	}
 	NewOutput->ExtraValues.Emplace(TEXT("platform"), FVariant(TEXT("dx")));
-	NewOutput->ExtraValues.Emplace(TEXT("dxversion"), FVariant((int64) DecoderPlatformHandle->GetDXVersionTimes1000()));
+	NewOutput->ExtraValues.Emplace(TEXT("dxversion"), FVariant((int64)DecoderPlatformHandle->GetDXVersionTimes1000()));
 	NewOutput->ExtraValues.Emplace(TEXT("sw"), FVariant(DecoderPlatformHandle->IsSoftware()));
 	NewOutput->ExtraValues.Emplace(TEXT("codec"), FVariant(TEXT("hevc")));
 	NewOutput->ExtraValues.Emplace(TEXT("pixfmt"), FVariant((int64)PixFmt));
 	NewOutput->ExtraValues.Emplace(TEXT("pixenc"), FVariant((int64)EElectraDecoderPlatformPixelEncoding::Native));
 
-	// Retain the IMFSample in the output. It is needed later in converting it for display.
-	NewOutput->MFSample = DecodedOutputSample;
+#if ELECTRA_HAVE_IMAGEBUFFERS
+	void* PlatformDevice = nullptr;
+	int32 PlatformDeviceVersion = 0;
+	bool bUseGPUBuffers = false;
+	auto PinnedResourceDelegate = ResourceDelegate.Pin();
+	if (PinnedResourceDelegate.IsValid())
+	{
+		PinnedResourceDelegate->GetD3DDevice(&PlatformDevice, &PlatformDeviceVersion);
+		bUseGPUBuffers = (PlatformDevice && PlatformDeviceVersion >= 12000);
+	}
+
+	// If we don't have support for DX12 direct to video decoding, yet, we bring the CPU buffer data into a DX12 resource now
+	if (bUseGPUBuffers)
+	{
+		DWORD BuffersNum = 0;
+		if ((Result = DecodedOutputSample->GetBufferCount(&BuffersNum)) != S_OK)
+		{
+			return false;
+		}
+		if (BuffersNum != 1)
+		{
+			return false;
+		}
+		TRefCountPtr<IMFMediaBuffer> Buffer;
+		if ((Result = DecodedOutputSample->GetBufferByIndex(0, Buffer.GetInitReference())) != S_OK)
+		{
+			return false;
+		}
+
+#if !ALLOW_MFSAMPLE_WITH_DX12
+		//
+		// Pre SDK 22621: use CPU copy to get data from MFSample to DX12 resource
+		//
+		TRefCountPtr D3D12Device(static_cast<ID3D12Device*>(PlatformDevice));
+
+		// Create the resource pool as needed...
+		if (!D3D12ResourcePool)
+		{
+			D3D12ResourcePool = MakeShared<FElectraMediaDecoderOutputBufferPool_DX12, ESPMode::ThreadSafe>(D3D12Device, MaxOutputBuffers, MaxWidth, MaxHeight * 3 / 2, (PixFmt == EElectraDecoderPlatformPixelFormat::NV12) ? 1 : 2);
+		}
+
+		// Create a tasksync instance so we can have our own async jobs run in consecutive order
+		if (!TaskSync.IsValid())
+		{
+			TaskSync.Reset(PinnedResourceDelegate->CreateAsyncConsecutiveTaskSync());
+		}
+
+		// Get data from MFSample...
+		TRefCountPtr<IMF2DBuffer> Buffer2D;
+		if ((Result = Buffer->QueryInterface(__uuidof(IMF2DBuffer), (void**)Buffer2D.GetInitReference())) != S_OK)
+		{
+			return false;
+		}
+
+		// Request resource and fence...
+		uint32 BufferPitch;
+		D3D12ResourcePool->AllocateOutputDataAsBuffer(NewOutput->Buffer, BufferPitch);
+
+		NewOutput->Pitch = BufferPitch;
+
+		if (!PinnedResourceDelegate->RunCodeAsync([DecodedOutputSample, Buffer2D, BufferResource = NewOutput->Buffer.Resource, BufferFence = NewOutput->Buffer.Fence, BufferFenceValue = NewOutput->Buffer.FenceValue, BufferPitch, DecodedHeight = NewOutput->DecodedHeight]()
+			{
+				CopyBufferData(Buffer2D, BufferResource, BufferFence, BufferFenceValue, BufferPitch, DecodedHeight);
+			}, TaskSync.Get()))
+		{
+			CopyBufferData(Buffer2D, NewOutput->Buffer.Resource, NewOutput->Buffer.Fence, NewOutput->Buffer.FenceValue, BufferPitch, NewOutput->DecodedHeight);
+		}
+#else
+		//
+		// Post SDK 22621: direct link from MFSample to DX12 resource
+		//
+		
+		TRefCountPtr<IMFDXGIBuffer> DXGIBuffer;
+		if ((Result = Buffer->QueryInterface(__uuidof(IMFDXGIBuffer), (void**)DXGIBuffer.GetInitReference())) != S_OK)
+		{
+			return false;
+		}
+
+		TRefCountPtr<ID3D12Resource> Resource;
+		Result = DXGIBuffer->GetResource(__uuidof(ID3D12Resource), (void**)Resource.GetInitReference());
+		if (Result == S_OK)
+		{
+			// Yes, set this up...
+			TRefCountPtr<IMFD3D12SynchronizationObjectCommands> DecoderSync;
+			if ((Result = DXGIBuffer->GetUnknown(MF_D3D12_SYNCHRONIZATION_OBJECT, IID_PPV_ARGS(DecoderSync.GetInitReference()))) != S_OK)
+			{
+				return false;
+			}
+
+			NewOutput->Buffer.Resource = Resource;
+			NewOutput->Buffer.DecoderSync = DecoderSync;
+			NewOutput->Buffer.MFSample = DecodedOutputSample;	// we keep a reference around until we actually submit the sync object to a DX12 queue
+		}
+#endif // !ALLOW_MFSAMPLE_WITH_DX12
+	}
+	else
+#endif // ELECTRA_HAVE_IMAGEBUFFERS
+	{
+		//
+		// Fallback output path for pre- / none-DX12 graphics APIs
+		// 
+		
+		// Retain the IMFSample in the output. It is needed later in converting it for display.
+		NewOutput->Pitch = stride;
+		NewOutput->MFSample = DecodedOutputSample;
+	}
+
 	// Do the same for the HW acceleration device.
 	NewOutput->Ctx.SetDeviceAndContext(DecoderPlatformHandle->GetDXDevice(), DecoderPlatformHandle->GetDXDeviceContext());
 
diff --git a/Engine/Plugins/Media/ElectraCodecs/Source/ElectraDecoders/Private/Windows/WindowsPlatformHeaders_Video_DX.h b/Engine/Plugins/Media/ElectraCodecs/Source/ElectraDecoders/Private/Windows/WindowsPlatformHeaders_Video_DX.h
--- a/Engine/Plugins/Media/ElectraCodecs/Source/ElectraDecoders/Private/Windows/WindowsPlatformHeaders_Video_DX.h
+++ b/Engine/Plugins/Media/ElectraCodecs/Source/ElectraDecoders/Private/Windows/WindowsPlatformHeaders_Video_DX.h
@@ -6,8 +6,8 @@
 #include "CoreMinimal.h"
 
 #include "Windows/AllowWindowsPlatformTypes.h"
-
 THIRD_PARTY_INCLUDES_START
+
 #include <d3d11.h>
 #include <d3d11_1.h>
 #include <d3d12.h>
@@ -17,10 +17,16 @@ THIRD_PARTY_INCLUDES_START
 #include "mfapi.h"
 #include "mferror.h"
 #include "mfidl.h"
-THIRD_PARTY_INCLUDES_END
 
-#include "Windows/HideWindowsPlatformTypes.h"
+#if defined(NTDDI_WIN10_NI)
+#include "mfd3d12.h"
+#define ALLOW_MFSAMPLE_WITH_DX12	1	// Windows SDK 22621 and up do feature APIs to support DX12 texture resources with WMF transforms
+#else
+#define ALLOW_MFSAMPLE_WITH_DX12	0
+#endif
 
+THIRD_PARTY_INCLUDES_END
+#include "Windows/HideWindowsPlatformTypes.h"
 
 struct FElectraVideoDecoderDXDeviceContext
 {
diff --git a/Engine/Plugins/Media/ElectraCodecs/Source/ElectraDecoders/Private/Windows/h264/H264_VideoDecoder_Windows.cpp b/Engine/Plugins/Media/ElectraCodecs/Source/ElectraDecoders/Private/Windows/h264/H264_VideoDecoder_Windows.cpp
--- a/Engine/Plugins/Media/ElectraCodecs/Source/ElectraDecoders/Private/Windows/h264/H264_VideoDecoder_Windows.cpp
+++ b/Engine/Plugins/Media/ElectraCodecs/Source/ElectraDecoders/Private/Windows/h264/H264_VideoDecoder_Windows.cpp
@@ -333,43 +333,60 @@ IElectraDecoder::FError IElectraVideoDecoderH264_DX::PlatformCreateMFDecoderTran
 					return Error.SetMessage(TEXT("MFCreateDXGIDeviceManager() failed")).SetCode(ERRCODE_INTERNAL_BAD_PARAMETER).SetSdkCode(hr);
 				}
 
-#ifdef ELECTRA_DECODERS_HAVE_DX11
-				D3D_FEATURE_LEVEL FeatureLevel;
-
-				uint32 DeviceCreationFlags = 0;
-				if ((DxDeviceCreationFlags & D3D11_CREATE_DEVICE_DEBUG) != 0)
+				bool bAssociatedDeviceManager = false;
+#if ALLOW_MFSAMPLE_WITH_DX12
+				// Try setting the D3D12 device with the device manager.
+				if (pfh->DXVersion == 12000)
 				{
-					DeviceCreationFlags |= D3D11_CREATE_DEVICE_DEBUG;
+					bAssociatedDeviceManager = SUCCEEDED(pfh->DxDeviceManager->ResetDevice(static_cast<ID3D12Device*>(ApplicationD3DDevice), ResetToken));
+					// note: pfh->Ctx.DxDevice & pfh->Ctx.DxDeviceContext will be nullptr (everything uses the DX12 render device)
 				}
+#endif
 
-				hr = D3D11CreateDevice(DXGIAdapter, DXGIAdapter ? D3D_DRIVER_TYPE_UNKNOWN : D3D_DRIVER_TYPE_HARDWARE, nullptr,
-						DeviceCreationFlags, nullptr, 0, D3D11_SDK_VERSION, pfh->Ctx.DxDevice.GetInitReference(), &FeatureLevel, pfh->Ctx.DxDeviceContext.GetInitReference());
-				if (hr != S_OK || !pfh->Ctx.DxDevice.IsValid())
-				{
-					return Error.SetMessage(TEXT("D3D11CreateDevice() failed")).SetCode(ERRCODE_INTERNAL_BAD_PARAMETER).SetSdkCode(hr);
-				}
-				if (FeatureLevel < D3D_FEATURE_LEVEL_9_3)
-				{
-					return Error.SetMessage(TEXT("Failed to create D3D11 Device with feature level 9.3 or above")).SetCode(ERRCODE_INTERNAL_BAD_PARAMETER);
-				}
-				if ((hr = pfh->DxDeviceManager->ResetDevice(pfh->Ctx.DxDevice, ResetToken)) != S_OK)
+#ifdef ELECTRA_DECODERS_HAVE_DX11
+				if (!bAssociatedDeviceManager)
 				{
-					return Error.SetMessage(TEXT("DXGIDeviceManager::ResetDevice() failed")).SetCode(ERRCODE_INTERNAL_BAD_PARAMETER).SetSdkCode(hr);
-				}
+					D3D_FEATURE_LEVEL FeatureLevel;
+
+					uint32 DeviceCreationFlags = 0;
+					if ((DxDeviceCreationFlags & D3D11_CREATE_DEVICE_DEBUG) != 0)
+					{
+						DeviceCreationFlags |= D3D11_CREATE_DEVICE_DEBUG;
+					}
+
+					hr = D3D11CreateDevice(DXGIAdapter, DXGIAdapter ? D3D_DRIVER_TYPE_UNKNOWN : D3D_DRIVER_TYPE_HARDWARE, nullptr,
+							DeviceCreationFlags, nullptr, 0, D3D11_SDK_VERSION, pfh->Ctx.DxDevice.GetInitReference(), &FeatureLevel, pfh->Ctx.DxDeviceContext.GetInitReference());
+					if (hr != S_OK || !pfh->Ctx.DxDevice.IsValid())
+					{
+						return Error.SetMessage(TEXT("D3D11CreateDevice() failed")).SetCode(ERRCODE_INTERNAL_BAD_PARAMETER).SetSdkCode(hr);
+					}
+					if (FeatureLevel < D3D_FEATURE_LEVEL_9_3)
+					{
+						return Error.SetMessage(TEXT("Failed to create D3D11 Device with feature level 9.3 or above")).SetCode(ERRCODE_INTERNAL_BAD_PARAMETER);
+					}
+
+					hr = pfh->DxDeviceManager->ResetDevice(pfh->Ctx.DxDevice, ResetToken);
+					if (hr != S_OK)
+					{
+						return Error.SetMessage(TEXT("DXGIDeviceManager::ResetDevice() failed")).SetCode(ERRCODE_INTERNAL_BAD_PARAMETER).SetSdkCode(hr);
+					}
+					// Multithread protect the newly created device as we're going to use it from decoding thread and from render thread for texture
+					// sharing between decoding and rendering DX devices
+					TRefCountPtr<ID3D10Multithread> DxMultithread;
+					if ((hr = pfh->Ctx.DxDevice->QueryInterface(__uuidof(ID3D10Multithread), (void**)DxMultithread.GetInitReference())) == S_OK)
+					{
+						DxMultithread->SetMultithreadProtected(1);
+					}
+					bAssociatedDeviceManager = true;
+
+
+
 
-				// Multithread protect the newly created device as we're going to use it from decoding thread and from render thread for texture
-				// sharing between decoding and rendering DX devices
-				TRefCountPtr<ID3D10Multithread> DxMultithread;
-				if ((hr = pfh->Ctx.DxDevice->QueryInterface(__uuidof(ID3D10Multithread), (void**)DxMultithread.GetInitReference())) == S_OK)
-				{
-					DxMultithread->SetMultithreadProtected(1);
 				}
 #endif
 
 				// No need to set a device manager when we didn't create a device for it to manage.
-#ifdef ELECTRA_DECODERS_HAVE_DX11
-				if ((hr = NewDecoder->ProcessMessage(MFT_MESSAGE_SET_D3D_MANAGER, reinterpret_cast<ULONG_PTR>(pfh->DxDeviceManager.GetReference()))) != S_OK)
-#endif
+				if (!bAssociatedDeviceManager || (hr = NewDecoder->ProcessMessage(MFT_MESSAGE_SET_D3D_MANAGER, reinterpret_cast<ULONG_PTR>(pfh->DxDeviceManager.GetReference()))) != S_OK)
 				{
 					// Fallback to software decoding.
 					pfh->bIsSW = true;
diff --git a/Engine/Plugins/Media/ElectraCodecs/Source/ElectraDecoders/Private/Windows/h265/H265_VideoDecoder_Windows.cpp b/Engine/Plugins/Media/ElectraCodecs/Source/ElectraDecoders/Private/Windows/h265/H265_VideoDecoder_Windows.cpp
--- a/Engine/Plugins/Media/ElectraCodecs/Source/ElectraDecoders/Private/Windows/h265/H265_VideoDecoder_Windows.cpp
+++ b/Engine/Plugins/Media/ElectraCodecs/Source/ElectraDecoders/Private/Windows/h265/H265_VideoDecoder_Windows.cpp
@@ -308,43 +308,56 @@ IElectraDecoder::FError IElectraVideoDecoderH265_DX::PlatformCreateMFDecoderTran
 					return Error.SetMessage(TEXT("MFCreateDXGIDeviceManager() failed")).SetCode(ERRCODE_INTERNAL_BAD_PARAMETER).SetSdkCode(hr);
 				}
 
-#ifdef ELECTRA_DECODERS_HAVE_DX11
-				D3D_FEATURE_LEVEL FeatureLevel;
-
-				uint32 DeviceCreationFlags = 0;
-				if ((DxDeviceCreationFlags & D3D11_CREATE_DEVICE_DEBUG) != 0)
+				// Try setting the D3D12 device with the device manager.
+				bool bAssociatedDeviceManager = false;
+#if ALLOW_MFSAMPLE_WITH_DX12
+				if (pfh->DXVersion >= 12000)
 				{
-					DeviceCreationFlags |= D3D11_CREATE_DEVICE_DEBUG;
+					bAssociatedDeviceManager = SUCCEEDED(pfh->DxDeviceManager->ResetDevice(static_cast<ID3D12Device*>(ApplicationD3DDevice), ResetToken));
+					// note: pfh->Ctx.DxDevice & pfh->Ctx.DxDeviceContext will be nullptr (everything uses the DX12 render device)
 				}
+#endif
 
-				hr = D3D11CreateDevice(DXGIAdapter, DXGIAdapter ? D3D_DRIVER_TYPE_UNKNOWN : D3D_DRIVER_TYPE_HARDWARE, nullptr,
-						DeviceCreationFlags, nullptr, 0, D3D11_SDK_VERSION, pfh->Ctx.DxDevice.GetInitReference(), &FeatureLevel, pfh->Ctx.DxDeviceContext.GetInitReference());
-				if (hr != S_OK || !pfh->Ctx.DxDevice.IsValid())
-				{
-					return Error.SetMessage(TEXT("D3D11CreateDevice() failed")).SetCode(ERRCODE_INTERNAL_BAD_PARAMETER).SetSdkCode(hr);
-				}
-				if (FeatureLevel < D3D_FEATURE_LEVEL_9_3)
-				{
-					return Error.SetMessage(TEXT("Failed to create D3D11 Device with feature level 9.3 or above")).SetCode(ERRCODE_INTERNAL_BAD_PARAMETER);
-				}
-				if ((hr = pfh->DxDeviceManager->ResetDevice(pfh->Ctx.DxDevice, ResetToken)) != S_OK)
+#ifdef ELECTRA_DECODERS_HAVE_DX11
+				if (!bAssociatedDeviceManager)
 				{
-					return Error.SetMessage(TEXT("DXGIDeviceManager::ResetDevice() failed")).SetCode(ERRCODE_INTERNAL_BAD_PARAMETER).SetSdkCode(hr);
-				}
+					D3D_FEATURE_LEVEL FeatureLevel;
 
-				// Multithread protect the newly created device as we're going to use it from decoding thread and from render thread for texture
-				// sharing between decoding and rendering DX devices
-				TRefCountPtr<ID3D10Multithread> DxMultithread;
-				if ((hr = pfh->Ctx.DxDevice->QueryInterface(__uuidof(ID3D10Multithread), (void**)DxMultithread.GetInitReference())) == S_OK)
-				{
-					DxMultithread->SetMultithreadProtected(1);
+					uint32 DeviceCreationFlags = 0;
+					if ((DxDeviceCreationFlags & D3D11_CREATE_DEVICE_DEBUG) != 0)
+					{
+						DeviceCreationFlags |= D3D11_CREATE_DEVICE_DEBUG;
+					}
+
+					hr = D3D11CreateDevice(DXGIAdapter, DXGIAdapter ? D3D_DRIVER_TYPE_UNKNOWN : D3D_DRIVER_TYPE_HARDWARE, nullptr,
+							DeviceCreationFlags, nullptr, 0, D3D11_SDK_VERSION, pfh->Ctx.DxDevice.GetInitReference(), &FeatureLevel, pfh->Ctx.DxDeviceContext.GetInitReference());
+					if (hr != S_OK || !pfh->Ctx.DxDevice.IsValid())
+					{
+						return Error.SetMessage(TEXT("D3D11CreateDevice() failed")).SetCode(ERRCODE_INTERNAL_BAD_PARAMETER).SetSdkCode(hr);
+					}
+					if (FeatureLevel < D3D_FEATURE_LEVEL_9_3)
+					{
+						return Error.SetMessage(TEXT("Failed to create D3D11 Device with feature level 9.3 or above")).SetCode(ERRCODE_INTERNAL_BAD_PARAMETER);
+					}
+
+					hr = pfh->DxDeviceManager->ResetDevice(pfh->Ctx.DxDevice, ResetToken);
+					if (hr != S_OK)
+					{
+						return Error.SetMessage(TEXT("DXGIDeviceManager::ResetDevice() failed")).SetCode(ERRCODE_INTERNAL_BAD_PARAMETER).SetSdkCode(hr);
+					}
+					// Multithread protect the newly created device as we're going to use it from decoding thread and from render thread for texture
+					// sharing between decoding and rendering DX devices
+					TRefCountPtr<ID3D10Multithread> DxMultithread;
+					if ((hr = pfh->Ctx.DxDevice->QueryInterface(__uuidof(ID3D10Multithread), (void**)DxMultithread.GetInitReference())) == S_OK)
+					{
+						DxMultithread->SetMultithreadProtected(1);
+					}
+					bAssociatedDeviceManager = true;
 				}
 #endif
 
 				// No need to set a device manager when we didn't create a device for it to manage.
-#ifdef ELECTRA_DECODERS_HAVE_DX11
-				if ((hr = NewDecoder->ProcessMessage(MFT_MESSAGE_SET_D3D_MANAGER, reinterpret_cast<ULONG_PTR>(pfh->DxDeviceManager.GetReference()))) != S_OK)
-#endif
+				if (!bAssociatedDeviceManager || (hr = NewDecoder->ProcessMessage(MFT_MESSAGE_SET_D3D_MANAGER, reinterpret_cast<ULONG_PTR>(pfh->DxDeviceManager.GetReference()))) != S_OK)
 				{
 					// Fallback to software decoding.
 					pfh->bIsSW = true;
diff --git a/Engine/Plugins/Media/ElectraCodecs/Source/ElectraDecoders/Public/Android/AndroidElectraDecoderGPUBufferHelpers.h b/Engine/Plugins/Media/ElectraCodecs/Source/ElectraDecoders/Public/Android/AndroidElectraDecoderGPUBufferHelpers.h
new file mode 100644
--- /dev/null
+++ b/Engine/Plugins/Media/ElectraCodecs/Source/ElectraDecoders/Public/Android/AndroidElectraDecoderGPUBufferHelpers.h
@@ -0,0 +1,6 @@
+// Copyright Epic Games, Inc. All Rights Reserved.
+
+#pragma once
+
+#include "CommonElectraDecoderGPUBufferHelpers.h"
+
diff --git a/Engine/Plugins/Media/ElectraCodecs/Source/ElectraDecoders/Public/Apple/AppleElectraDecoderGPUBufferHelpers.h b/Engine/Plugins/Media/ElectraCodecs/Source/ElectraDecoders/Public/Apple/AppleElectraDecoderGPUBufferHelpers.h
new file mode 100644
--- /dev/null
+++ b/Engine/Plugins/Media/ElectraCodecs/Source/ElectraDecoders/Public/Apple/AppleElectraDecoderGPUBufferHelpers.h
@@ -0,0 +1,6 @@
+// Copyright Epic Games, Inc. All Rights Reserved.
+
+#pragma once
+
+#include "CommonElectraDecoderGPUBufferHelpers.h"
+
diff --git a/Engine/Plugins/Media/ElectraCodecs/Source/ElectraDecoders/Public/Apple/AppleElectraDecoderPlatformOutputHandleTypes.h b/Engine/Plugins/Media/ElectraCodecs/Source/ElectraDecoders/Public/Apple/AppleElectraDecoderPlatformOutputHandleTypes.h
--- a/Engine/Plugins/Media/ElectraCodecs/Source/ElectraDecoders/Public/Apple/AppleElectraDecoderPlatformOutputHandleTypes.h
+++ b/Engine/Plugins/Media/ElectraCodecs/Source/ElectraDecoders/Public/Apple/AppleElectraDecoderPlatformOutputHandleTypes.h
@@ -40,6 +40,13 @@ enum class EElectraDecoderPlatformPixelEncoding
 	ARGB_BigEndian,	//!< Interpret as ARGB, big endian
 };
 
+class FElectraDecoderOutputSync final
+{
+public:
+	FElectraDecoderOutputSync() = default;
+	~FElectraDecoderOutputSync() = default;
+};
+
 class IElectraDecoderVideoOutputImageBuffers
 {
 public:
@@ -55,6 +62,9 @@ public:
 	// Returns the n'th CFImageBuffer ("texture") reference.
 	virtual void* GetBufferTextureByIndex(int32 InBufferIndex) const = 0;
 
+	// Returns the n'th image buffer's GPU sync object to signal texture data as being ready to be read
+	virtual bool GetBufferTextureSyncByIndex(int32 InBufferIndex, FElectraDecoderOutputSync& SyncObject) const = 0;
+
 	// Returns the n'th image buffer format
 	virtual EElectraDecoderPlatformPixelFormat GetBufferFormatByIndex(int32 InBufferIndex) const = 0;
 
diff --git a/Engine/Plugins/Media/ElectraCodecs/Source/ElectraDecoders/Public/CommonElectraDecoderGPUBufferHelpers.h b/Engine/Plugins/Media/ElectraCodecs/Source/ElectraDecoders/Public/CommonElectraDecoderGPUBufferHelpers.h
new file mode 100644
--- /dev/null
+++ b/Engine/Plugins/Media/ElectraCodecs/Source/ElectraDecoders/Public/CommonElectraDecoderGPUBufferHelpers.h
@@ -0,0 +1,13 @@
+// Copyright Epic Games, Inc. All Rights Reserved.
+
+#pragma once
+
+// Provide fallback defaults for various defines
+#ifndef ELECTRA_MEDIAGPUBUFFER_DX12
+#define	ELECTRA_MEDIAGPUBUFFER_DX12 0
+#endif
+
+// Decoder settings define the amount of output samples in flight, but this setting omits a few internal instances that
+// will occur at the same time. Usually just one, but at times 2, this informs this constant used to make sure pre-allocated
+// output buffer heaps are large enough at all times
+static const uint32 kElectraDecoderPipelineExtraFrames = 2;
diff --git a/Engine/Plugins/Media/ElectraCodecs/Source/ElectraDecoders/Public/IElectraDecoderResourceDelegateBase.h b/Engine/Plugins/Media/ElectraCodecs/Source/ElectraDecoders/Public/IElectraDecoderResourceDelegateBase.h
--- a/Engine/Plugins/Media/ElectraCodecs/Source/ElectraDecoders/Public/IElectraDecoderResourceDelegateBase.h
+++ b/Engine/Plugins/Media/ElectraCodecs/Source/ElectraDecoders/Public/IElectraDecoderResourceDelegateBase.h
@@ -28,6 +28,22 @@ public:
 		Video
 	};
 
+	/**
+	* Interface class to abstract async job synchronization primitives on platforms implementing them
+	*/
+	class IAsyncConsecutiveTaskSync
+	{
+	public:
+		virtual ~IAsyncConsecutiveTaskSync() = default;
+	protected:
+		IAsyncConsecutiveTaskSync() = default;
+	};
+
+	virtual IAsyncConsecutiveTaskSync* CreateAsyncConsecutiveTaskSync()
+	{ return nullptr; }
+	virtual bool RunCodeAsync(TFunction<void()>&& CodeToRun, IAsyncConsecutiveTaskSync* TaskSync = nullptr)
+	{ return false; }
+
 	virtual IDecoderPlatformResource* CreatePlatformResource(void* InOwnerHandle, EDecoderPlatformResourceType InDecoderResourceType, const TMap<FString, FVariant> InOptions)
 	{ return nullptr; }
 	virtual void ReleasePlatformResource(void* InOwnerHandle, IDecoderPlatformResource* InHandleToDestroy)
diff --git a/Engine/Plugins/Media/ElectraCodecs/Source/ElectraDecoders/Public/IOS/IOSElectraDecoderGPUBufferHelpers.h b/Engine/Plugins/Media/ElectraCodecs/Source/ElectraDecoders/Public/IOS/IOSElectraDecoderGPUBufferHelpers.h
new file mode 100644
--- /dev/null
+++ b/Engine/Plugins/Media/ElectraCodecs/Source/ElectraDecoders/Public/IOS/IOSElectraDecoderGPUBufferHelpers.h
@@ -0,0 +1,5 @@
+// Copyright Epic Games, Inc. All Rights Reserved.
+
+#pragma once
+
+#include "Apple/AppleElectraDecoderGPUBufferHelpers.h"
diff --git a/Engine/Plugins/Media/ElectraCodecs/Source/ElectraDecoders/Public/Linux/LinuxElectraDecoderGPUBufferHelpers.h b/Engine/Plugins/Media/ElectraCodecs/Source/ElectraDecoders/Public/Linux/LinuxElectraDecoderGPUBufferHelpers.h
new file mode 100644
--- /dev/null
+++ b/Engine/Plugins/Media/ElectraCodecs/Source/ElectraDecoders/Public/Linux/LinuxElectraDecoderGPUBufferHelpers.h
@@ -0,0 +1,5 @@
+// Copyright Epic Games, Inc. All Rights Reserved.
+
+#pragma once
+
+#include "CommonElectraDecoderGPUBufferHelpers.h"
diff --git a/Engine/Plugins/Media/ElectraCodecs/Source/ElectraDecoders/Public/Linux/LinuxElectraDecoderPlatformOutputHandleTypes.h b/Engine/Plugins/Media/ElectraCodecs/Source/ElectraDecoders/Public/Linux/LinuxElectraDecoderPlatformOutputHandleTypes.h
--- a/Engine/Plugins/Media/ElectraCodecs/Source/ElectraDecoders/Public/Linux/LinuxElectraDecoderPlatformOutputHandleTypes.h
+++ b/Engine/Plugins/Media/ElectraCodecs/Source/ElectraDecoders/Public/Linux/LinuxElectraDecoderPlatformOutputHandleTypes.h
@@ -41,6 +41,13 @@ enum class EElectraDecoderPlatformPixelEncoding
 	ARGB_BigEndian,	//!< Interpret as ARGB, big endian
 };
 
+class FElectraDecoderOutputSync final
+{
+public:
+	FElectraDecoderOutputSync() = default;
+	~FElectraDecoderOutputSync() = default;
+};
+
 class IElectraDecoderVideoOutputImageBuffers
 {
 public:
@@ -56,6 +63,9 @@ public:
 	// Returns the n'th image buffer as GPU texture resource.
 	virtual void* GetBufferTextureByIndex(int32 InBufferIndex) const = 0;
 
+	// Returns the n'th image buffer's GPU sync object to signal texture data as being ready to be read
+	virtual bool GetBufferTextureSyncByIndex(int32 InBufferIndex, FElectraDecoderOutputSync& SyncObject) const = 0;
+
 	// Returns the n'th image buffer format
 	virtual EElectraDecoderPlatformPixelFormat GetBufferFormatByIndex(int32 InBufferIndex) const = 0;
 
diff --git a/Engine/Plugins/Media/ElectraCodecs/Source/ElectraDecoders/Public/Mac/MacElectraDecoderGPUBufferHelpers.h b/Engine/Plugins/Media/ElectraCodecs/Source/ElectraDecoders/Public/Mac/MacElectraDecoderGPUBufferHelpers.h
new file mode 100644
--- /dev/null
+++ b/Engine/Plugins/Media/ElectraCodecs/Source/ElectraDecoders/Public/Mac/MacElectraDecoderGPUBufferHelpers.h
@@ -0,0 +1,5 @@
+// Copyright Epic Games, Inc. All Rights Reserved.
+
+#pragma once
+
+#include "Apple/AppleElectraDecoderGPUBufferHelpers.h"
diff --git a/Engine/Plugins/Media/ElectraCodecs/Source/ElectraDecoders/Public/TVOS/TVOSElectraDecoderGPUBufferHelpers.h b/Engine/Plugins/Media/ElectraCodecs/Source/ElectraDecoders/Public/TVOS/TVOSElectraDecoderGPUBufferHelpers.h
new file mode 100644
--- /dev/null
+++ b/Engine/Plugins/Media/ElectraCodecs/Source/ElectraDecoders/Public/TVOS/TVOSElectraDecoderGPUBufferHelpers.h
@@ -0,0 +1,5 @@
+// Copyright Epic Games, Inc. All Rights Reserved.
+
+#pragma once
+
+#include "Apple/AppleElectraDecoderGPUBufferHelpers.h"
diff --git a/Engine/Plugins/Media/ElectraCodecs/Source/ElectraDecoders/Public/Unix/UnixElectraDecoderGPUBufferHelpers.h b/Engine/Plugins/Media/ElectraCodecs/Source/ElectraDecoders/Public/Unix/UnixElectraDecoderGPUBufferHelpers.h
new file mode 100644
--- /dev/null
+++ b/Engine/Plugins/Media/ElectraCodecs/Source/ElectraDecoders/Public/Unix/UnixElectraDecoderGPUBufferHelpers.h
@@ -0,0 +1,5 @@
+// Copyright Epic Games, Inc. All Rights Reserved.
+
+#pragma once
+
+#include "Linux/LinuxElectraDecoderGPUBufferHelpers.h"
diff --git a/Engine/Plugins/Media/ElectraCodecs/Source/ElectraDecoders/Public/Windows/WindowsElectraDecoderGPUBufferHelpers.h b/Engine/Plugins/Media/ElectraCodecs/Source/ElectraDecoders/Public/Windows/WindowsElectraDecoderGPUBufferHelpers.h
new file mode 100644
--- /dev/null
+++ b/Engine/Plugins/Media/ElectraCodecs/Source/ElectraDecoders/Public/Windows/WindowsElectraDecoderGPUBufferHelpers.h
@@ -0,0 +1,364 @@
+// Copyright Epic Games, Inc. All Rights Reserved.
+
+#pragma once
+
+#include "CoreMinimal.h"
+
+/*********************************************************************************************************************/
+/*********************************************************************************************************************/
+/*********************************************************************************************************************/
+
+#include "Windows/AllowWindowsPlatformTypes.h"
+THIRD_PARTY_INCLUDES_START
+#include "mfobjects.h"
+#include "d3d12.h"
+#include "d3dx12.h"
+THIRD_PARTY_INCLUDES_END
+#include "Windows/HideWindowsPlatformTypes.h"
+
+#define	ELECTRA_MEDIAGPUBUFFER_DX12 1
+
+/*********************************************************************************************************************/
+/*********************************************************************************************************************/
+/*********************************************************************************************************************/
+
+#ifdef _WIN64
+#define ALLOW_USE_OF_SSEMEMCPY 1
+
+#include <intrin.h>
+
+static void SSE2MemCpy(const void* InDst, const void* InSrc, unsigned int InSize)
+{
+	__m128i* Dst = (__m128i*)InDst;
+	__m128i* Src = (__m128i*)InSrc;
+
+	InSize = InSize >> 7;
+	for (unsigned int i = 0; i < InSize; i++)
+	{
+		_mm_prefetch((char*)(Src + 8), 0);
+		_mm_prefetch((char*)(Src + 10), 0);
+		_mm_prefetch((char*)(Src + 12), 0);
+		_mm_prefetch((char*)(Src + 14), 0);
+		__m128i m0 = _mm_load_si128(Src + 0);
+		__m128i m1 = _mm_load_si128(Src + 1);
+		__m128i m2 = _mm_load_si128(Src + 2);
+		__m128i m3 = _mm_load_si128(Src + 3);
+		__m128i m4 = _mm_load_si128(Src + 4);
+		__m128i m5 = _mm_load_si128(Src + 5);
+		__m128i m6 = _mm_load_si128(Src + 6);
+		__m128i m7 = _mm_load_si128(Src + 7);
+		_mm_stream_si128(Dst + 0, m0);
+		_mm_stream_si128(Dst + 1, m1);
+		_mm_stream_si128(Dst + 2, m2);
+		_mm_stream_si128(Dst + 3, m3);
+		_mm_stream_si128(Dst + 4, m4);
+		_mm_stream_si128(Dst + 5, m5);
+		_mm_stream_si128(Dst + 6, m6);
+		_mm_stream_si128(Dst + 7, m7);
+		Src += 8;
+		Dst += 8;
+	}
+}
+
+#else
+#define ALLOW_USE_OF_SSEMEMCPY 0
+#endif // _WIN64
+
+/*********************************************************************************************************************/
+/*********************************************************************************************************************/
+/*********************************************************************************************************************/
+
+// Simple DX12 heap and fence manager for use with decoder output buffers (both upload heaps or default GPU heaps)
+class FElectraMediaDecoderOutputBufferPool_DX12 : public TSharedFromThis<FElectraMediaDecoderOutputBufferPool_DX12, ESPMode::ThreadSafe>
+{
+	struct FResourceInfo
+	{
+		FResourceInfo(TWeakPtr<FElectraMediaDecoderOutputBufferPool_DX12> InPool, uint32 InBufferIdx, TRefCountPtr<ID3D12Heap> InHeap) : Pool(InPool), BufferIdx(InBufferIdx), Heap(InHeap) {}
+
+		TWeakPtr<FElectraMediaDecoderOutputBufferPool_DX12> Pool;
+		uint32 BufferIdx;
+		TRefCountPtr<ID3D12Heap> Heap;
+	};
+
+public:
+	// Create instance for use with buffers (usually for uploading data)
+	FElectraMediaDecoderOutputBufferPool_DX12(TRefCountPtr<ID3D12Device> InD3D12Device, uint32 InMaxNumBuffers, uint32 Width, uint32 Height, uint32 BytesPerPixel, D3D12_HEAP_TYPE InD3D12HeapType = D3D12_HEAP_TYPE_UPLOAD)
+		: D3D12Device(InD3D12Device)
+		, D3D12HeapType(InD3D12HeapType)
+		, MaxNumBuffers(InMaxNumBuffers)
+		, FreeMask((1 << InMaxNumBuffers) - 1)
+		, LastFenceValue(0)
+	{
+		check(InMaxNumBuffers <= 32);
+
+		// Get the size needed per buffer
+		BufferPitch = Align(Width * BytesPerPixel, D3D12_TEXTURE_DATA_PITCH_ALIGNMENT);
+		BufferSize = Align(BufferPitch * Height, FMath::Max(D3D12_TEXTURE_DATA_PLACEMENT_ALIGNMENT, D3D12_DEFAULT_RESOURCE_PLACEMENT_ALIGNMENT));
+
+		InitCommon();
+	}
+
+	// Create instance for use with textures (usually to receive any upload buffers or internal decoder texture data)
+	FElectraMediaDecoderOutputBufferPool_DX12(TRefCountPtr<ID3D12Device> InD3D12Device, uint32 InMaxNumBuffers, uint32 Width, uint32 Height, DXGI_FORMAT PixFmt, D3D12_HEAP_TYPE InD3D12HeapType = D3D12_HEAP_TYPE_DEFAULT)
+		: D3D12Device(InD3D12Device)
+		, D3D12HeapType(InD3D12HeapType)
+		, MaxNumBuffers(InMaxNumBuffers)
+		, FreeMask((1 << InMaxNumBuffers) - 1)
+		, LastFenceValue(0)
+	{
+		check(InMaxNumBuffers <= 32);
+
+		// Get the size needed per buffer / texture
+		D3D12_RESOURCE_DESC Desc = {};
+		Desc.MipLevels = 1;
+		Desc.Format = PixFmt;
+		Desc.Width = Width;
+		Desc.Height = Height;
+		Desc.Flags = D3D12_RESOURCE_FLAG_NONE;
+		Desc.DepthOrArraySize = 1;
+		Desc.SampleDesc.Count = 1;
+		Desc.SampleDesc.Quality = 0;
+		Desc.Layout = D3D12_TEXTURE_LAYOUT_UNKNOWN;
+		Desc.Dimension = D3D12_RESOURCE_DIMENSION_TEXTURE2D;
+		D3D12_RESOURCE_ALLOCATION_INFO AllocInfo = InD3D12Device->GetResourceAllocationInfo(0, 1, &Desc);
+
+		BufferPitch = 0;
+		BufferSize = Align(AllocInfo.SizeInBytes, AllocInfo.Alignment);
+		InitCommon();
+	}
+
+	~FElectraMediaDecoderOutputBufferPool_DX12()
+	{
+	}
+
+	// Check if the current setup is compatible with the new parameters
+	bool IsCompatibleAsBuffer(uint32 InMaxNumBuffers, uint32 Width, uint32 Height, uint32 BytesPerPixel) const
+	{
+		uint32 NewBufferPitch = Align(Width * BytesPerPixel, D3D12_TEXTURE_DATA_PITCH_ALIGNMENT);
+		uint32 NewBufferSize = Align(NewBufferPitch * Height, FMath::Max(D3D12_TEXTURE_DATA_PLACEMENT_ALIGNMENT, D3D12_DEFAULT_RESOURCE_PLACEMENT_ALIGNMENT));
+		return (MaxNumBuffers == InMaxNumBuffers && BufferSize == NewBufferSize && BufferPitch == NewBufferPitch);
+	}
+
+	// Check if a buffer is available
+	bool BufferAvailable() const
+	{
+		return FPlatformAtomics::AtomicRead((const int32*)&FreeMask) != 0;
+	}
+
+	class FOutputData
+	{
+	public:
+		~FOutputData()
+		{
+			check(ReadyToDestroy());
+		}
+
+		TRefCountPtr<ID3D12Resource> Resource;
+		TRefCountPtr<ID3D12Fence> Fence;
+		uint64 FenceValue;
+
+		bool ReadyToDestroy() const
+		{
+			bool bOk = true;
+#if !UE_BUILD_SHIPPING
+			if (Resource.IsValid())
+			{
+				bOk = Resource->AddRef() > 1 || Fence->GetCompletedValue() >= FenceValue;
+				Resource->Release();
+			}
+#endif
+			return bOk;
+		}
+	};
+
+	// Allocate a buffer resource and return suitable sync fence data
+	bool AllocateOutputDataAsBuffer(FOutputData& OutData, uint32& OutPitch)
+	{
+		OutData.Resource = AllocateBuffer(OutPitch, [NumBytes = BufferSize](D3D12_RESOURCE_DESC& Desc)
+		{
+			Desc.MipLevels = 1;
+			Desc.Format = DXGI_FORMAT_UNKNOWN;
+			Desc.Width = NumBytes;
+			Desc.Height = 1;
+			Desc.Flags = D3D12_RESOURCE_FLAG_NONE;
+			Desc.DepthOrArraySize = 1;
+			Desc.SampleDesc.Count = 1;
+			Desc.SampleDesc.Quality = 0;
+			Desc.Layout = D3D12_TEXTURE_LAYOUT_ROW_MAJOR;
+			Desc.Dimension = D3D12_RESOURCE_DIMENSION_BUFFER;
+		});
+		OutData.Fence = GetUpdatedBufferFence(OutData.FenceValue);
+		return OutData.Resource.IsValid();
+	}
+
+	// Allocate a texture resource and return suitable sync fence data
+	bool AllocateOutputDataAsTexture(FOutputData& OutData, uint32 InWidth, uint32 InHeight, DXGI_FORMAT InPixFmt)
+	{
+		uint32 OutPitch;
+		OutData.Resource = AllocateBuffer(OutPitch, [InWidth, InHeight, InPixFmt](D3D12_RESOURCE_DESC& Desc)
+		{
+			Desc.MipLevels = 1;
+			Desc.Format = InPixFmt;
+			Desc.Width = InWidth;
+			Desc.Height = InHeight;
+			Desc.Flags = D3D12_RESOURCE_FLAG_NONE;
+			Desc.DepthOrArraySize = 1;
+			Desc.SampleDesc.Count = 1;
+			Desc.SampleDesc.Quality = 0;
+			Desc.Layout = D3D12_TEXTURE_LAYOUT_UNKNOWN;
+			Desc.Dimension = D3D12_RESOURCE_DIMENSION_TEXTURE2D;
+		});
+		OutData.Fence = GetUpdatedBufferFence(OutData.FenceValue);
+		return OutData.Resource.IsValid();
+	}
+
+	TRefCountPtr<ID3D12Fence> GetUpdatedBufferFence(uint64& FenceValue)
+	{
+		FenceValue = ++LastFenceValue;
+		return D3D12BufferFence;
+	}
+
+	// Helper function to copy simple, linear texture data between differently picthed buffers
+	static void CopyWithPitchAdjust(uint8* Dst, uint32 DstPitch, const uint8* Src, uint32 SrcPitch, uint32 NumRows)
+	{
+		if (DstPitch != SrcPitch)
+		{
+			for (uint32 Y = NumRows; Y > 0; --Y)
+			{
+				FMemory::Memcpy(Dst, Src, SrcPitch);
+				Src += SrcPitch;
+				Dst += DstPitch;
+			}
+		}
+		else
+		{
+			uint32 NumBytes = SrcPitch * NumRows;
+#if ALLOW_USE_OF_SSEMEMCPY
+			if ((NumBytes & 0x7f) == 0)
+			{
+				SSE2MemCpy(Dst, Src, NumBytes);
+			}
+			else
+#endif
+			{
+				FMemory::Memcpy(Dst, Src, NumBytes);
+			}
+		}
+	}
+
+private:
+	void InitCommon()
+	{
+		D3D12_HEAP_DESC HeapDesc = {};
+		HeapDesc.Alignment = D3D12_DEFAULT_RESOURCE_PLACEMENT_ALIGNMENT;
+		HeapDesc.SizeInBytes = BufferSize * MaxNumBuffers;
+		HeapDesc.Properties = CD3DX12_HEAP_PROPERTIES(D3D12HeapType);
+		HeapDesc.Flags = (D3D12HeapType == D3D12_HEAP_TYPE_UPLOAD) ? D3D12_HEAP_FLAG_ALLOW_ONLY_BUFFERS : D3D12_HEAP_FLAG_ALLOW_ALL_BUFFERS_AND_TEXTURES;
+		HRESULT Res = D3D12Device->CreateHeap(&HeapDesc, IID_PPV_ARGS(D3D12OutputHeap.GetInitReference()));
+		check(SUCCEEDED(Res));
+		if (Res != S_OK)
+		{
+			FreeMask = 0;
+		}
+#if !UE_BUILD_SHIPPING
+		D3D12OutputHeap->SetName(TEXT("ElectraOutputBufferPoolHeap"));
+#endif
+
+		Res = D3D12Device->CreateFence(0, D3D12_FENCE_FLAG_NONE, IID_PPV_ARGS(D3D12BufferFence.GetInitReference()));
+		check(SUCCEEDED(Res));
+#if !UE_BUILD_SHIPPING
+		D3D12BufferFence->SetName(TEXT("ElectraOutputBufferPoolFence"));
+#endif
+	}
+
+	TRefCountPtr<ID3D12Resource> AllocateBuffer(uint32& OutBufferPitch, TFunction<void(D3D12_RESOURCE_DESC& Desc)> && InitializeDesc)
+	{
+		TRefCountPtr<ID3D12Resource> Resource;
+
+		// Find the buffer index we can use...
+		int32 BufferIdx, OldFreeMask, NewFreeMask;
+		do
+		{
+			OldFreeMask = (int32)FreeMask;
+			if (OldFreeMask == 0)
+			{
+				BufferIdx = -1;
+				break;
+			}
+			BufferIdx = 31 - FMath::CountLeadingZeros(OldFreeMask);
+			NewFreeMask = OldFreeMask & ~(1 << BufferIdx);
+		} while (OldFreeMask != FPlatformAtomics::InterlockedCompareExchange((int32*)&FreeMask, NewFreeMask, OldFreeMask));
+
+		// Anything?
+		if (BufferIdx >= 0)
+		{
+			// Make a placed resource to use...
+			D3D12_RESOURCE_DESC Desc = {};
+			InitializeDesc(Desc);
+			HRESULT Res = D3D12Device->CreatePlacedResource(D3D12OutputHeap, BufferSize * BufferIdx, &Desc, (D3D12HeapType == D3D12_HEAP_TYPE_UPLOAD) ? D3D12_RESOURCE_STATE_GENERIC_READ : D3D12_RESOURCE_STATE_COMMON, nullptr, IID_PPV_ARGS(Resource.GetInitReference()));
+			check(SUCCEEDED(Res));
+			if (Res == S_OK)
+			{
+#if !UE_BUILD_SHIPPING
+				Resource->SetName(TEXT("ElectraOutputBufferPoolBufferResource"));
+#endif
+
+				OutBufferPitch = BufferPitch;
+
+				// Register a destruction callback so we can reset our memory management here transparently...
+				TRefCountPtr<ID3DDestructionNotifier> Notifier;
+				Res = Resource->QueryInterface(__uuidof(ID3DDestructionNotifier), (void**)Notifier.GetInitReference());
+				check(SUCCEEDED(Res));
+
+				// note: we keep a reference to the heap in our context data for the destruction callback, so we can ensure the heap lives as long as there are placed resources in it
+				UINT CallbackID;
+				Res = Notifier->RegisterDestructionCallback(ResourceDestructionCallback, new FResourceInfo(AsWeak(), BufferIdx, D3D12OutputHeap), &CallbackID);
+				check(SUCCEEDED(Res));
+			}
+			else
+			{
+				FreeBuffer(BufferIdx);
+			}
+		}
+
+		return Resource;
+	}
+
+	void FreeBuffer(int32 InBufferIdx)
+	{
+		// Reset buffer index to free...
+		int32 OldFreeMask, NewFreeMask;
+		do
+		{
+			OldFreeMask = (int32)FreeMask;
+			NewFreeMask = OldFreeMask | (1 << InBufferIdx);
+		} while (OldFreeMask != FPlatformAtomics::InterlockedCompareExchange((int32*)&FreeMask, NewFreeMask, OldFreeMask));
+	}
+
+	static void ResourceDestructionCallback(void* Context)
+	{
+		auto ResourceInfo = reinterpret_cast<FResourceInfo*>(Context);
+
+		if (auto Pool = ResourceInfo->Pool.Pin())
+		{
+			Pool->FreeBuffer(ResourceInfo->BufferIdx);
+		}
+
+		// Get rid of our resource tracking info block...
+		// (this will also release the ref to the heap the resource was on)
+		delete ResourceInfo;
+	}
+
+	TRefCountPtr<ID3D12Device> D3D12Device;
+	D3D12_HEAP_TYPE D3D12HeapType;
+	TRefCountPtr<ID3D12Heap> D3D12OutputHeap;
+	TRefCountPtr<ID3D12Fence> D3D12BufferFence;
+	uint32 MaxNumBuffers;
+	uint32 BufferSize;
+	uint32 BufferPitch;
+	uint32 FreeMask;
+	uint64 LastFenceValue;
+};
+
+
+#include "CommonElectraDecoderGPUBufferHelpers.h"
diff --git a/Engine/Plugins/Media/ElectraCodecs/Source/ElectraDecoders/Public/Windows/WindowsElectraDecoderPlatformOutputHandleTypes.h b/Engine/Plugins/Media/ElectraCodecs/Source/ElectraDecoders/Public/Windows/WindowsElectraDecoderPlatformOutputHandleTypes.h
--- a/Engine/Plugins/Media/ElectraCodecs/Source/ElectraDecoders/Public/Windows/WindowsElectraDecoderPlatformOutputHandleTypes.h
+++ b/Engine/Plugins/Media/ElectraCodecs/Source/ElectraDecoders/Public/Windows/WindowsElectraDecoderPlatformOutputHandleTypes.h
@@ -4,6 +4,13 @@
 
 #include "CoreMinimal.h"
 
+#include "Windows/AllowWindowsPlatformTypes.h"
+THIRD_PARTY_INCLUDES_START
+#include "d3d12.h"
+#include "mfobjects.h"
+THIRD_PARTY_INCLUDES_END
+#include "Windows/HideWindowsPlatformTypes.h"
+
 enum class EElectraDecoderPlatformOutputHandleType
 {
 	MFSample,
@@ -43,6 +50,27 @@ enum class EElectraDecoderPlatformPixelEncoding
 	ARGB_BigEndian,	//!< Interpret as ARGB, big endian
 };
 
+class FElectraDecoderOutputSync final
+{
+public:
+	FElectraDecoderOutputSync() : SyncValue(0), CopyDoneSyncValue(0) {}
+	FElectraDecoderOutputSync(TRefCountPtr<IUnknown> InSync, uint64 InSyncValue, TRefCountPtr<IMFSample> InMFSample = nullptr)
+		: Sync(InSync), SyncValue(InSyncValue), MFSample(InMFSample)
+	{}
+	FElectraDecoderOutputSync(TRefCountPtr<IUnknown> InSync, uint64 InSyncValue, TRefCountPtr<ID3D12Fence> InCopyDoneSync, uint64 InCopyDoneSyncValue)
+		: Sync(InSync), SyncValue(InSyncValue), CopyDoneSync(InCopyDoneSync), CopyDoneSyncValue(InCopyDoneSyncValue)
+	{}
+
+	TRefCountPtr<IUnknown>		Sync;				// An IUnknown based sync primitive (fence, or DX12 specific WMF object)
+	uint64						SyncValue;			// Optionally needed sync value
+
+	TRefCountPtr<ID3D12Fence>	CopyDoneSync;		// Fence to signal end of copy from this buffer back to decoder
+	uint64						CopyDoneSyncValue;	// Value to be used with CopyDoneSync fence
+
+private:
+	TRefCountPtr<IMFSample>		MFSample;			// The sample reference can ride along to ensure that the DX12 WMF interface properly keeps reference to the sample until we enqueue their sync primitives with the a D3D queue
+};
+
 class IElectraDecoderVideoOutputImageBuffers
 {
 public:
@@ -58,6 +86,9 @@ public:
 	// Returns the n'th image buffer as GPU texture resource.
 	virtual void* GetBufferTextureByIndex(int32 InBufferIndex) const = 0;
 
+	// Returns the n'th image buffer's GPU sync object to signal texture data as being ready to be read
+	virtual bool GetBufferTextureSyncByIndex(int32 InBufferIndex, FElectraDecoderOutputSync& SyncObject) const = 0;
+
 	// Returns the n'th image buffer format
 	virtual EElectraDecoderPlatformPixelFormat GetBufferFormatByIndex(int32 InBufferIndex) const = 0;
 
@@ -67,3 +98,5 @@ public:
 	// Returns the n'th image buffer pitch
 	virtual int32 GetBufferPitchByIndex(int32 InBufferIndex) const = 0;
 };
+
+#define ELECTRA_HAVE_IMAGEBUFFERS 1
diff --git a/Engine/Plugins/Media/ElectraPlayer/Source/ElectraPlayerPlugin/ElectraPlayerPlugin.Build.cs b/Engine/Plugins/Media/ElectraPlayer/Source/ElectraPlayerPlugin/ElectraPlayerPlugin.Build.cs
--- a/Engine/Plugins/Media/ElectraPlayer/Source/ElectraPlayerPlugin/ElectraPlayerPlugin.Build.cs
+++ b/Engine/Plugins/Media/ElectraPlayer/Source/ElectraPlayerPlugin/ElectraPlayerPlugin.Build.cs
@@ -1,4 +1,5 @@
 // Copyright Epic Games, Inc. All Rights Reserved.
+using System.IO;
 using UnrealBuildTool;
 
 namespace UnrealBuildTool.Rules
@@ -22,7 +23,8 @@ namespace UnrealBuildTool.Rules
 					"RHI",
 					"ElectraPlayerRuntime",
 					"ElectraSamples",
-					"ElectraBase"
+					"ElectraBase",
+					"ElectraDecoders"
 				});
 
 			PrivateIncludePathModuleNames.AddRange(
@@ -35,6 +37,12 @@ namespace UnrealBuildTool.Rules
 				PrivateDependencyModuleNames.Add("Engine");
 			}
 
+			if (Target.Platform.IsInGroup(UnrealPlatformGroup.Windows))
+			{
+				PrivateDependencyModuleNames.Add("D3D12RHI");
+				AddEngineThirdPartyPrivateStaticDependencies(Target, "DX12");
+			}
+
 			if (Target.IsInPlatformGroup(UnrealPlatformGroup.Apple))
 			{
 				PrivateDependencyModuleNames.Add("MetalRHI");
diff --git a/Engine/Plugins/Media/ElectraPlayer/Source/ElectraPlayerPlugin/Private/ElectraPlayerPluginModule.cpp b/Engine/Plugins/Media/ElectraPlayer/Source/ElectraPlayerPlugin/Private/ElectraPlayerPluginModule.cpp
--- a/Engine/Plugins/Media/ElectraPlayer/Source/ElectraPlayerPlugin/Private/ElectraPlayerPluginModule.cpp
+++ b/Engine/Plugins/Media/ElectraPlayer/Source/ElectraPlayerPlugin/Private/ElectraPlayerPluginModule.cpp
@@ -11,11 +11,14 @@
 #include "ElectraPlayerPlugin.h"
 #include "ParameterDictionary.h"
 #include "SimpleElectraAudioPlayer.h"
+#include "IElectraDecoderResourceDelegateBase.h"
 
 #define LOCTEXT_NAMESPACE "ElectraPlayerPluginModule"
 
 DEFINE_LOG_CATEGORY(LogElectraPlayerPlugin);
 
+DECLARE_CYCLE_STAT(TEXT("Electra AsyncJob"), STAT_ElectraAsyncJob, STATGROUP_Media);
+
 // -----------------------------------------------------------------------------------------------------------------------------------
 
 class FElectraPlayerPluginModule : public IElectraPlayerPluginModule
@@ -80,6 +83,36 @@ public:
 		}
 	}
 
+	class FAsyncConsecutiveTaskSync : public IElectraDecoderResourceDelegateBase::IAsyncConsecutiveTaskSync
+	{
+	public:
+		FAsyncConsecutiveTaskSync() {}
+		virtual ~FAsyncConsecutiveTaskSync() {}
+
+		FGraphEventRef	GraphEvent;
+	};
+
+	static IElectraDecoderResourceDelegateBase::IAsyncConsecutiveTaskSync* CreateAsyncConsecutiveTaskSync()
+	{
+		return new FAsyncConsecutiveTaskSync();
+	}
+
+	static void RunCodeAsync(TFunction<void()>&& CodeToRun, IElectraDecoderResourceDelegateBase::IAsyncConsecutiveTaskSync* TaskSync)
+	{
+		FScopeLock Lock(&AsyncJobAccessCS);
+
+		// Execute code async
+		// (We assume this to be code copying buffer data around. Hence we allow only ONE at any time to not clog up
+		//  the buses even more and delay the copy process further)
+		FGraphEventArray Events;
+		auto& Event = TaskSync ? static_cast<FAsyncConsecutiveTaskSync*>(TaskSync)->GraphEvent : RunCodeAsyncEvent;
+		if (Event.IsValid())
+		{
+			Events.Add(Event);
+		}
+		Event = FFunctionGraphTask::CreateAndDispatchWhenReady(MoveTemp(CodeToRun), GET_STATID(STAT_ElectraAsyncJob), &Events);
+	}
+	
 	void StartupModule() override
 	{
 		// Check that we have the player module and that it has initialized successfully.
@@ -108,6 +141,8 @@ public:
 
 			Electra::FParamDict Params;
 			Params.Set(FName(TEXT("GetDeviceTypeCallback")), Electra::FVariantValue((void*)&FElectraPlayerPluginModule::GetDynamicRHIInfo));
+			Params.Set(FName(TEXT("CreateAsyncConsecutiveTaskSync")), Electra::FVariantValue((void*)&FElectraPlayerPluginModule::CreateAsyncConsecutiveTaskSync));
+			Params.Set(FName(TEXT("RunCodeAsyncCallback")), Electra::FVariantValue((void*)&FElectraPlayerPluginModule::RunCodeAsync));
 			if (!FElectraPlayerPlatform::StartupPlatformResources(Params))
 			{
 				UE_LOG(LogElectraPlayerPlugin, Log, TEXT("Platform resource setup failed! Electra Player plugin is not initialised."));
@@ -130,8 +165,14 @@ private:
 	FElectraPlayerSendAnalyticMetricsPerMinuteDelegate	SendAnalyticMetricsPerMinuteDelegate;
 	FElectraPlayerReportVideoStreamingErrorDelegate		ReportVideoStreamingErrorDelegate;
 	FElectraPlayerReportSubtitlesMetricsDelegate		ReportSubtitlesMetricsDelegate;
+
+	static FCriticalSection AsyncJobAccessCS;
+	static FGraphEventRef RunCodeAsyncEvent;
 };
 
+FCriticalSection FElectraPlayerPluginModule::AsyncJobAccessCS;
+FGraphEventRef FElectraPlayerPluginModule::RunCodeAsyncEvent;
+
 IMPLEMENT_MODULE(FElectraPlayerPluginModule, ElectraPlayerPlugin);
 
 #undef LOCTEXT_NAMESPACE
diff --git a/Engine/Plugins/Media/ElectraPlayer/Source/ElectraPlayerPlugin/Private/Windows/ElectraPlayerPlugin_Windows.cpp b/Engine/Plugins/Media/ElectraPlayer/Source/ElectraPlayerPlugin/Private/Windows/ElectraPlayerPlugin_Windows.cpp
--- a/Engine/Plugins/Media/ElectraPlayer/Source/ElectraPlayerPlugin/Private/Windows/ElectraPlayerPlugin_Windows.cpp
+++ b/Engine/Plugins/Media/ElectraPlayer/Source/ElectraPlayerPlugin/Private/Windows/ElectraPlayerPlugin_Windows.cpp
@@ -6,9 +6,26 @@
 
 #include <ElectraPlayerPlugin.h>
 
+#include "ID3D12DynamicRHI.h"
+
+class FPlayerResourceDelegateWindows : public IElectraPlayerResourceDelegate
+{
+public:
+	FPlayerResourceDelegateWindows()
+	{
+	}
+
+	virtual void ExecuteCodeWithCopyCommandQueueUsage(TFunction<void(ID3D12CommandQueue*)>&& CodeToRun) override
+	{
+		GetID3D12PlatformDynamicRHI()->RHIRunOnQueue(ED3D12RHIRunOnQueueType::Copy, MoveTemp(CodeToRun), false);
+	}
+
+private:
+};
+
 IElectraPlayerResourceDelegate* FElectraPlayerPlugin::PlatformCreatePlayerResourceDelegate()
 {
-	return nullptr;
+	return new FPlayerResourceDelegateWindows();
 }
 
 #endif // PLATFORM_WINDOWS
diff --git a/Engine/Plugins/Media/ElectraPlayer/Source/ElectraPlayerRuntime/ElectraPlayerRuntime.Build.cs b/Engine/Plugins/Media/ElectraPlayer/Source/ElectraPlayerRuntime/ElectraPlayerRuntime.Build.cs
--- a/Engine/Plugins/Media/ElectraPlayer/Source/ElectraPlayerRuntime/ElectraPlayerRuntime.Build.cs
+++ b/Engine/Plugins/Media/ElectraPlayer/Source/ElectraPlayerRuntime/ElectraPlayerRuntime.Build.cs
@@ -48,6 +48,7 @@ namespace UnrealBuildTool.Rules
 				PrivateDefinitions.Add("_CRT_SECURE_NO_WARNINGS=1");
 
 				AddEngineThirdPartyPrivateStaticDependencies(Target, "WinHttp");
+				AddEngineThirdPartyPrivateStaticDependencies(Target, "DX12");
 
 				PrivateDefinitions.Add("ELECTRA_HAVE_DX11");	// video decoding for DX11 enabled (Win8+)
 
diff --git a/Engine/Plugins/Media/ElectraPlayer/Source/ElectraPlayerRuntime/Private/Runtime/Decoder/Windows/ElectraVideoDecoder_PC.cpp b/Engine/Plugins/Media/ElectraPlayer/Source/ElectraPlayerRuntime/Private/Runtime/Decoder/Windows/ElectraVideoDecoder_PC.cpp
--- a/Engine/Plugins/Media/ElectraPlayer/Source/ElectraPlayerRuntime/Private/Runtime/Decoder/Windows/ElectraVideoDecoder_PC.cpp
+++ b/Engine/Plugins/Media/ElectraPlayer/Source/ElectraPlayerRuntime/Private/Runtime/Decoder/Windows/ElectraVideoDecoder_PC.cpp
@@ -2,6 +2,29 @@
 
 #include "ElectraVideoDecoder_PC.h"
 #include "Renderer/RendererVideo.h"
+#include COMPILED_PLATFORM_HEADER(ElectraDecoderGPUBufferHelpers.h)
+#include COMPILED_PLATFORM_HEADER(ElectraDecoderPlatformOutputHandleTypes.h)
+//#include COMPILED_PLATFORM_HEADER(PlatformHeaders_Video_DX.h) // PRIVATE FOR CODECS
+
+#include "Windows/AllowWindowsPlatformTypes.h"
+THIRD_PARTY_INCLUDES_START
+
+#include <d3d12.h>
+#include "mftransform.h"
+#include "mfapi.h"
+#include "mferror.h"
+#include "mfidl.h"
+#if defined(NTDDI_WIN10_NI)
+#include "mfd3d12.h"
+#define ALLOW_MFSAMPLE_WITH_DX12	1	// Windows SDK 22621 and up do feature APIs to support DX12 texture resources with WMF transforms
+#else
+#define ALLOW_MFSAMPLE_WITH_DX12	0
+#endif
+
+THIRD_PARTY_INCLUDES_END
+#include "Windows/HideWindowsPlatformTypes.h"
+
+// ------------------------------------------------------------------------------------
 
 void FElectraPlayerVideoDecoderOutputPC::SetOwner(const TSharedPtr<IDecoderOutputOwner, ESPMode::ThreadSafe>& InOwningRenderer)
 {
@@ -13,12 +36,6 @@ FElectraPlayerVideoDecoderOutputPC::EOutputType FElectraPlayerVideoDecoderOutput
 	return OutputType;
 }
 
-TRefCountPtr<IMFSample> FElectraPlayerVideoDecoderOutputPC::GetMFSample() const
-{
-	check(OutputType == EOutputType::SoftwareWin8Plus || OutputType == EOutputType::SoftwareWin7);
-	return MFSample;
-}
-
 const TArray<uint8>& FElectraPlayerVideoDecoderOutputPC::GetBuffer() const
 {
 	check(OutputType == EOutputType::SoftwareWin8Plus || OutputType == EOutputType::SoftwareWin7 || OutputType == EOutputType::HardwareDX9_DX12);
@@ -40,7 +57,7 @@ uint32 FElectraPlayerVideoDecoderOutputPC::GetStride() const
 
 TRefCountPtr<IUnknown> FElectraPlayerVideoDecoderOutputPC::GetTexture() const
 {
-	return static_cast<IUnknown*>(Texture.GetReference());
+	return Texture ? static_cast<IUnknown*>(Texture.GetReference()) : static_cast<IUnknown*>(TextureDX12.GetReference());
 }
 
 TRefCountPtr<ID3D11Device> FElectraPlayerVideoDecoderOutputPC::GetDevice() const
@@ -56,9 +73,19 @@ FIntPoint FElectraPlayerVideoDecoderOutputPC::GetDim() const
 
 TRefCountPtr<IUnknown> FElectraPlayerVideoDecoderOutputPC::GetSync(uint64& SyncValue) const
 {
-	TRefCountPtr<IUnknown> KeyedMutex;
-	HRESULT res = Texture->QueryInterface(_uuidof(IDXGIKeyedMutex), (void**)&KeyedMutex);
-	return KeyedMutex;
+	if (Texture)
+	{
+		TRefCountPtr<IUnknown> KeyedMutex;
+		HRESULT res = Texture->QueryInterface(_uuidof(IDXGIKeyedMutex), (void**)&KeyedMutex);
+		SyncValue = 0;
+		return KeyedMutex;
+	}
+	else if (TextureDX12)
+	{
+		SyncValue = FenceValue;
+		return static_cast<IUnknown*>(D3DFence.GetReference());
+	}
+	return nullptr;
 }
 
 void FElectraPlayerVideoDecoderOutputPC::InitializeWithBuffer(const void* InBuffer, uint32 InSize, uint32 InStride, FIntPoint Dim, TSharedPtr<Electra::FParamDict, ESPMode::ThreadSafe> InParamDict)
@@ -86,6 +113,301 @@ void FElectraPlayerVideoDecoderOutputPC::InitializeWithBuffer(TSharedPtr<TArray<
 	Stride = InStride;
 }
 
+void FElectraPlayerVideoDecoderOutputPC::InitializeWithResource(const TRefCountPtr<ID3D12Device>& InD3D12Device, const TRefCountPtr<ID3D12Resource> InResource, uint32 ResourcePitch, const FElectraDecoderOutputSync& OutputSync, const FIntPoint& InOutputDim, TSharedPtr<Electra::FParamDict, ESPMode::ThreadSafe> InParamDict, Electra::IVideoDecoderResourceDelegate* InResourceDelegate,
+																uint32 MaxWidth, uint32 MaxHeight, uint32 MaxOutputBuffers)
+{
+	// General initialization
+	FVideoDecoderOutput::Initialize(InParamDict);
+
+	// Local reference to the resource passed in as we need to ensure it is alive until we have copied from it
+	DecoderOutputResource = InResource;
+
+	SampleDim = InOutputDim;
+	OutputType = EOutputType::Hardware_DX;
+
+	HRESULT Res;
+
+	// Create D3D12 resources as needed (we will reuse if possible as this instance is reused)
+	if (!D3DCmdAllocator.IsValid())
+	{
+		Res = InD3D12Device->CreateCommandAllocator(D3D12_COMMAND_LIST_TYPE_COPY, IID_PPV_ARGS(D3DCmdAllocator.GetInitReference()));
+		check(!FAILED(Res));
+#if !UE_BUILD_SHIPPING
+		D3DCmdAllocator->SetName(TEXT("ElectraVideoDecoderCmdAllocator"));
+#endif
+	}
+	if (!D3DCmdList.IsValid())
+	{
+		Res = InD3D12Device->CreateCommandList(0, D3D12_COMMAND_LIST_TYPE_COPY, D3DCmdAllocator.GetReference(), nullptr, __uuidof(ID3D12CommandList), reinterpret_cast<void**>(D3DCmdList.GetInitReference()));
+		check(!FAILED(Res));
+		D3DCmdList->Close();
+#if !UE_BUILD_SHIPPING
+		D3DCmdList->SetName(TEXT("ElectraVideoDecoderCmdList"));
+#endif
+	}
+
+	// - - - - - - - - - - - - - - - - - - - - - - - - - - - - - - - - - - - - - - - - - - - - - - - - - - - - - - - -
+
+	// Create texture to copy decoder output buffer into as needed
+
+	D3D12_RESOURCE_DESC ResourceDesc;
+	ResourceDesc = InResource->GetDesc();
+
+	const bool bIsTexture = (ResourceDesc.Dimension == D3D12_RESOURCE_DIMENSION_TEXTURE2D);
+
+	const FIntPoint Dim = GetDim();
+	const EPixelFormat PixFmt = GetFormat();
+	const EVideoDecoderPixelEncoding FmtEnc = GetFormatEncoding();
+
+	// Find out if the image data is sRGB encoded (we also assume this by default)
+	bool bOutputIsSRGB = true;
+	if (auto Colorimetry = GetColorimetry())
+	{
+		if (auto MPEGDef = Colorimetry->GetMPEGDefinition())
+		{
+			bOutputIsSRGB = (MPEGDef->TransferCharacteristics == 2);
+		}
+	}
+
+	// Check if we can use HW sRGB decode (we must avoid doing this if the data represents something else as true RGB(A) values)
+	const bool bIsSRGB = bOutputIsSRGB && (FmtEnc == EVideoDecoderPixelEncoding::Native || FmtEnc == EVideoDecoderPixelEncoding::RGB || FmtEnc == EVideoDecoderPixelEncoding::RGBA);
+
+	// note: this somewhat duplicates / bypasses the GPixelFormat array, which we COULD access here to get the DXGI formats.
+	//		 BUT! The platform format is only inserted by RHI and we do NOT rely ion RHI being operational!
+	uint32 BlockSizeX = 1;
+	uint32 BlockSizeY = 1;
+	DXGI_FORMAT DXGIFmt = DXGI_FORMAT_UNKNOWN;
+	switch (PixFmt)
+	{
+		case PF_NV12:				DXGIFmt = DXGI_FORMAT_NV12;				break;
+		case PF_P010:				DXGIFmt = DXGI_FORMAT_P010;				break;
+		case PF_A16B16G16R16:		DXGIFmt = DXGI_FORMAT_R16G16B16A16_UNORM; break;
+		case PF_R16G16B16A16_UNORM:	DXGIFmt = DXGI_FORMAT_R16G16B16A16_UNORM; break;
+		case PF_A32B32G32R32F:		DXGIFmt = DXGI_FORMAT_R32G32B32A32_FLOAT; break;
+		case PF_B8G8R8A8:			DXGIFmt = bIsSRGB ? DXGI_FORMAT_B8G8R8A8_UNORM_SRGB : DXGI_FORMAT_B8G8R8A8_TYPELESS;  break;
+		case PF_R8G8B8A8:			DXGIFmt = bIsSRGB ? DXGI_FORMAT_R8G8B8A8_UNORM_SRGB : DXGI_FORMAT_R8G8B8A8_TYPELESS;  break;
+		case PF_A2B10G10R10:		DXGIFmt = DXGI_FORMAT_R10G10B10A2_UNORM;  break;
+		case PF_DXT1:				DXGIFmt = bIsSRGB ? DXGI_FORMAT_BC1_UNORM_SRGB : DXGI_FORMAT_BC1_UNORM; BlockSizeX = BlockSizeY = 4; break;
+		case PF_DXT5:				DXGIFmt = bIsSRGB ? DXGI_FORMAT_BC3_UNORM_SRGB : DXGI_FORMAT_BC3_UNORM; BlockSizeX = BlockSizeY = 4; break;
+		case PF_BC4:				DXGIFmt = DXGI_FORMAT_BC4_UNORM; BlockSizeX = BlockSizeY = 4; break;
+		default:
+			check(!"Unexpected pixel format for D3D12 buffer resource!");
+	}
+
+	/*
+	* - We need to copy this from the output buffers ASAP to not throttle buffer availability (and hence performance) for the decoder
+	* - As RHI limits resource creation to the renderthread, which would delay creation, we need to use a platform API level texture we can create right here and now
+	*/
+	if (!TextureDX12.IsValid() || SampleDim != TextureDX12Dim)
+	{
+		// We get here only after the instance came back from the pool! Hence we can be sure any old texture is no longer actively used.
+		TextureDX12 = nullptr;
+
+#if !UE_BUILD_SHIPPING
+		if (bIsTexture)
+		{
+			// A little bit of sanity checking
+			check(DXGIFmt == ResourceDesc.Format);
+			check(SampleDim.X == ResourceDesc.Width);
+			check(SampleDim.Y == ResourceDesc.Height);
+		}
+#endif
+
+		if (!D3D12ResourcePool.IsValid())
+		{
+			D3D12ResourcePool = MakeShared<FElectraMediaDecoderOutputBufferPool_DX12>(InD3D12Device, MaxOutputBuffers + kElectraDecoderPipelineExtraFrames, MaxWidth, MaxHeight, DXGIFmt, D3D12_HEAP_TYPE_DEFAULT);
+			if (!D3D12ResourcePool.IsValid())
+			{
+				UE_LOG(LogElectraPlayer, Error, TEXT("Could not allocate texture resource heap!"));
+				return;
+			}
+		}
+		
+		FElectraMediaDecoderOutputBufferPool_DX12::FOutputData OutputData;
+		if (!D3D12ResourcePool->AllocateOutputDataAsTexture(OutputData, SampleDim.X, SampleDim.Y, DXGIFmt))
+		{
+			UE_LOG(LogElectraPlayer, Error, TEXT("Could not allocate texture resource from heap!"));
+			return;
+		}
+
+		TextureDX12 = OutputData.Resource;
+#if !UE_BUILD_SHIPPING
+		TextureDX12->SetName(TEXT("ElectraVideoDecoderOutputFrame"));
+#endif
+		D3DFence = OutputData.Fence;
+		FenceValue = OutputData.FenceValue;
+
+		TextureDX12Dim = SampleDim;
+	}
+	else
+	{
+		// We reuse the texture, but we need to update the fence info...
+		check(D3D12ResourcePool.IsValid());
+		D3DFence = D3D12ResourcePool->GetUpdatedBufferFence(FenceValue);
+	}
+
+	// - - - - - - - - - - - - - - - - - - - - - - - - - - - - - - - - - - - - - - - - - - - - - - - - - - - - - - - -
+
+	// Build command list to copy data from decoder buffer
+
+	D3D12_RESOURCE_BARRIER PreCopyResourceBarriers[] = {
+		{D3D12_RESOURCE_BARRIER_TYPE_TRANSITION, D3D12_RESOURCE_BARRIER_FLAG_NONE, {{ InResource, D3D12_RESOURCE_BARRIER_ALL_SUBRESOURCES, D3D12_RESOURCE_STATE_COMMON, D3D12_RESOURCE_STATE_COPY_SOURCE }}},
+		{D3D12_RESOURCE_BARRIER_TYPE_TRANSITION, D3D12_RESOURCE_BARRIER_FLAG_NONE, {{ TextureDX12, D3D12_RESOURCE_BARRIER_ALL_SUBRESOURCES, D3D12_RESOURCE_STATE_COMMON, D3D12_RESOURCE_STATE_COPY_DEST}}},
+	};
+
+	D3D12_RESOURCE_BARRIER PostCopyResourceBarriers[] = {
+		{ D3D12_RESOURCE_BARRIER_TYPE_TRANSITION, D3D12_RESOURCE_BARRIER_FLAG_NONE, {{ InResource, D3D12_RESOURCE_BARRIER_ALL_SUBRESOURCES, D3D12_RESOURCE_STATE_COPY_SOURCE, D3D12_RESOURCE_STATE_COMMON }}},
+		{ D3D12_RESOURCE_BARRIER_TYPE_TRANSITION, D3D12_RESOURCE_BARRIER_FLAG_NONE, {{ TextureDX12, D3D12_RESOURCE_BARRIER_ALL_SUBRESOURCES, D3D12_RESOURCE_STATE_COPY_DEST, D3D12_RESOURCE_STATE_COMMON }}},
+	};
+
+	D3DCmdAllocator->Reset();
+	D3DCmdList->Reset(D3DCmdAllocator, nullptr);
+
+	D3DCmdList->ResourceBarrier(2, PreCopyResourceBarriers);
+
+	if (bIsTexture)
+	{
+		// If we got a texture, we simply can copy the resource as a whole...
+		D3DCmdList->CopyResource(TextureDX12, InResource);
+	}
+	else
+	{
+		// Source is a buffer. Build a properly configured footprint and copy from there...
+		if (PixFmt == PF_NV12 || PixFmt == PF_P010)
+		{
+			//
+			// Formats with 2 planes
+			//
+
+			D3D12_BOX SrcBox;
+			SrcBox.left = 0;
+			SrcBox.top = 0;
+			SrcBox.right = Dim.X;
+			SrcBox.bottom = Dim.Y;
+			SrcBox.back = 1;
+			SrcBox.front = 0;
+
+			// Copy Y plane
+			D3D12_PLACED_SUBRESOURCE_FOOTPRINT SrcFootPrint{};
+			SrcFootPrint.Offset = 0;
+			SrcFootPrint.Footprint.Format = (PixFmt == PF_NV12) ? DXGI_FORMAT_R8_UNORM : DXGI_FORMAT_R16_UNORM;
+			SrcFootPrint.Footprint.Width = Dim.X;
+			SrcFootPrint.Footprint.Height = Dim.Y;
+			SrcFootPrint.Footprint.Depth = 1;
+			SrcFootPrint.Footprint.RowPitch = ResourcePitch;
+			CD3DX12_TEXTURE_COPY_LOCATION SourceCopyLocation(InResource, SrcFootPrint);
+			CD3DX12_TEXTURE_COPY_LOCATION DestCopyLocation(TextureDX12, 0);
+			D3DCmdList->CopyTextureRegion(&DestCopyLocation, 0, 0, 0, &SourceCopyLocation, &SrcBox);
+
+			// Copy CbCr plane
+			SrcBox.right = Dim.X >> 1;
+			SrcBox.bottom = Dim.Y >> 1;
+			SourceCopyLocation.PlacedFootprint.Offset = ResourcePitch * Dim.Y;
+			SourceCopyLocation.PlacedFootprint.Footprint.Format = (PixFmt == PF_NV12) ? DXGI_FORMAT_R8G8_UNORM : DXGI_FORMAT_R16G16_UNORM;
+			SourceCopyLocation.PlacedFootprint.Footprint.Width = Dim.X >> 1;
+			SourceCopyLocation.PlacedFootprint.Footprint.Height = Dim.Y >> 1;
+			DestCopyLocation.SubresourceIndex = 1;
+			D3DCmdList->CopyTextureRegion(&DestCopyLocation, 0, 0, 0, &SourceCopyLocation, &SrcBox);
+		}
+		else
+		{
+			//
+			// Single plane formats
+			//
+
+			D3D12_BOX SrcBox;
+			SrcBox.left = 0;
+			SrcBox.top = 0;
+			SrcBox.right = Dim.X;
+			SrcBox.bottom = Dim.Y;
+			SrcBox.back = 1;
+			SrcBox.front = 0;
+
+			D3D12_PLACED_SUBRESOURCE_FOOTPRINT SrcFootPrint{};
+			SrcFootPrint.Offset = 0;
+			SrcFootPrint.Footprint.Format = DXGIFmt;
+			SrcFootPrint.Footprint.Width = Align(Dim.X, BlockSizeX);
+			SrcFootPrint.Footprint.Height = Align(Dim.Y, BlockSizeY);
+			SrcFootPrint.Footprint.Depth = 1;
+			SrcFootPrint.Footprint.RowPitch = ResourcePitch;
+			CD3DX12_TEXTURE_COPY_LOCATION SourceCopyLocation(InResource, SrcFootPrint);
+			CD3DX12_TEXTURE_COPY_LOCATION DestCopyLocation(TextureDX12, 0);
+			D3DCmdList->CopyTextureRegion(&DestCopyLocation, 0, 0, 0, &SourceCopyLocation, &SrcBox);
+		}
+	}
+
+	D3DCmdList->ResourceBarrier(2, PostCopyResourceBarriers);
+
+	D3DCmdList->Close();
+
+	// - - - - - - - - - - - - - - - - - - - - - - - - - - - - - - - - - - - - - - - - - - - - - - - - - - - - - - - -
+
+	// Get all synchronization info we need
+	// (we could have a fence or a DX12 decoder synchronization object)
+
+	TRefCountPtr<ID3D12Fence> ResourceFence;
+#if ALLOW_MFSAMPLE_WITH_DX12
+	TRefCountPtr<IMFD3D12SynchronizationObjectCommands> DecoderSync;
+#else
+	TRefCountPtr<IUnknown> DecoderSync;
+#endif
+
+	if (OutputSync.Sync.IsValid())
+	{
+#if ALLOW_MFSAMPLE_WITH_DX12
+		if (OutputSync.Sync->QueryInterface(__uuidof(IMFD3D12SynchronizationObjectCommands), (void**)DecoderSync.GetInitReference()) != S_OK)
+#endif // ALLOW_MFSAMPLE_WITH_DX12
+		{
+			OutputSync.Sync->QueryInterface(__uuidof(ID3D12Fence), (void**)ResourceFence.GetInitReference());
+		}
+		check(ResourceFence.IsValid() || DecoderSync.IsValid());
+	}
+
+	// - - - - - - - - - - - - - - - - - - - - - - - - - - - - - - - - - - - - - - - - - - - - - - - - - - - - - - - -
+
+	// Trigger copy (this will eventually execute on the submission thread of RHI if running in UE)
+	// (note: we pass in all of FElectraDecoderOutputSync to guarantee any references needed to make the decoder output sync work are passed along, too!)
+	InResourceDelegate->ExecuteCodeWithCopyCommandQueueUsage([CmdList = D3DCmdList, DestFence = D3DFence, DestFenceValue = FenceValue, 
+#if ALLOW_MFSAMPLE_WITH_DX12
+															  DecoderSync, /*SAMPLE ?*/
+#endif
+															  ResourceFence, OutputSync](ID3D12CommandQueue* D3DCmdQueue)
+	{
+#if ALLOW_MFSAMPLE_WITH_DX12
+		if (DecoderSync)
+		{
+			// Sync queue to make sure decoder output is ready for us (WMF case)
+			DecoderSync->EnqueueResourceReadyWait(D3DCmdQueue);
+		}
+#endif
+		if (ResourceFence)
+		{
+			// Sync queue to make sure decoder output is ready for us
+			D3DCmdQueue->Wait(ResourceFence, OutputSync.SyncValue);
+		}
+
+		// Execute copy
+		ID3D12CommandList* CmdLists[1] = { CmdList.GetReference() };
+		D3DCmdQueue->ExecuteCommandLists(1, CmdLists);
+
+#if ALLOW_MFSAMPLE_WITH_DX12
+		if (DecoderSync)
+		{
+			// Sync to end of copy operation and release MFSample once reached
+			DecoderSync->EnqueueResourceRelease(D3DCmdQueue);
+		}
+#endif
+		// Trigger optional notification back to the decoder, so it could reuse its buffer after the copy
+		if (OutputSync.CopyDoneSync)
+		{
+			D3DCmdQueue->Signal(OutputSync.CopyDoneSync, OutputSync.CopyDoneSyncValue);
+		}
+
+		// Notify user of the output data we just copied of its arrival
+		D3DCmdQueue->Signal(DestFence, DestFenceValue);
+	});
+}
+
 
 void FElectraPlayerVideoDecoderOutputPC::InitializeWithSharedTexture(const TRefCountPtr<ID3D11Device>& InD3D11Device, const TRefCountPtr<ID3D11DeviceContext> InDeviceContext, const TRefCountPtr<IMFSample> InMFSample, const FIntPoint& OutputDim, TSharedPtr<Electra::FParamDict, ESPMode::ThreadSafe> InParamDict)
 {
@@ -130,6 +452,7 @@ void FElectraPlayerVideoDecoderOutputPC::InitializeWithSharedTexture(const TRefC
 		DecoderTexture->GetDesc(&DecoderTextureDesc);
 
 		// Make a texture we pass on as output
+		// (note: memory is NOT tracked in UE - much as other decoder related memory!)
 		D3D11_TEXTURE2D_DESC TextureDesc;
 		TextureDesc.Width = SampleDim.X;
 		TextureDesc.Height = SampleDim.Y;
@@ -187,6 +510,17 @@ void FElectraPlayerVideoDecoderOutputPC::InitializeWithSharedTexture(const TRefC
 	}
 }
 
+bool FElectraPlayerVideoDecoderOutputPC::IsReadyForReuse()
+{
+	// Make sure any async copies to any internal GPU resources are done (DX12 case)
+	// (this also guards any decoder created resources we keep to read from)
+	if (TextureDX12.IsValid() && D3DFence.IsValid())
+	{
+		return D3DFence->GetCompletedValue() >= FenceValue;
+	}
+	return true;
+}
+
 void FElectraPlayerVideoDecoderOutputPC::ShutdownPoolable()
 {
 	TSharedPtr<IDecoderOutputOwner, ESPMode::ThreadSafe> lockedVideoRenderer = OwningRenderer.Pin();
@@ -207,6 +541,12 @@ void FElectraPlayerVideoDecoderOutputPC::ShutdownPoolable()
 		Texture->QueryInterface(__uuidof(IDXGIResource), (void**)&OtherResource);
 	}
 
+	D3DFence = nullptr;
+
+	// Drop decoder output resource
+	DecoderOutputResource = nullptr;
+
+	// Make sure DX11 shared texture sync state is as expected
 	if (OtherResource)
 	{
 		HANDLE SharedHandle = nullptr;
diff --git a/Engine/Plugins/Media/ElectraPlayer/Source/ElectraPlayerRuntime/Private/Runtime/Decoder/Windows/ElectraVideoDecoder_PC.h b/Engine/Plugins/Media/ElectraPlayer/Source/ElectraPlayerRuntime/Private/Runtime/Decoder/Windows/ElectraVideoDecoder_PC.h
--- a/Engine/Plugins/Media/ElectraPlayer/Source/ElectraPlayerRuntime/Private/Runtime/Decoder/Windows/ElectraVideoDecoder_PC.h
+++ b/Engine/Plugins/Media/ElectraPlayer/Source/ElectraPlayerRuntime/Private/Runtime/Decoder/Windows/ElectraVideoDecoder_PC.h
@@ -5,6 +5,7 @@
 #include "PlayerCore.h"
 #include "ElectraPlayerPrivate.h"
 #include "MediaVideoDecoderOutputPC.h"
+#include "VideoDecoderResourceDelegate.h"
 
 /*********************************************************************************************************************/
 /*********************************************************************************************************************/
@@ -22,11 +23,16 @@ THIRD_PARTY_INCLUDES_END
 
 #include "Windows/HideWindowsPlatformTypes.h"
 
+struct ID3D12Resource;
+struct ID3D12Fence;
+
+class FElectraMediaDecoderOutputBufferPool_DX12;
+class FElectraDecoderOutputSync;
+
 /*********************************************************************************************************************/
 /*********************************************************************************************************************/
 /*********************************************************************************************************************/
 
-
 class FElectraPlayerVideoDecoderOutputPC : public FVideoDecoderOutputPC
 {
 public:
@@ -38,6 +44,9 @@ public:
 
 	void InitializeWithBuffer(TSharedPtr<TArray<uint8>, ESPMode::ThreadSafe> InBuffer, uint32 InStride, FIntPoint Dim, TSharedPtr<Electra::FParamDict, ESPMode::ThreadSafe> InParamDict);
 
+	void InitializeWithResource(const TRefCountPtr<ID3D12Device>& InD3D12Device, const TRefCountPtr<ID3D12Resource> Resource, uint32 ResourcePitch, const FElectraDecoderOutputSync& OutputSync, const FIntPoint& OutputDim, TSharedPtr<Electra::FParamDict, ESPMode::ThreadSafe> InParamDict, Electra::IVideoDecoderResourceDelegate* ResourceDelegate,
+								uint32 MaxWidth, uint32 MaxHeight, uint32 MaxOutputBuffers);
+
 	// Hardware decode to shared DX11 texture (Win8+) from IMFSample
 	void InitializeWithSharedTexture(const TRefCountPtr<ID3D11Device>& InD3D11Device, const TRefCountPtr<ID3D11DeviceContext> InDeviceContext, const TRefCountPtr<IMFSample> MFSample, const FIntPoint& OutputDim, TSharedPtr<Electra::FParamDict, ESPMode::ThreadSafe> InParamDict);
 
@@ -46,8 +55,8 @@ public:
 
 	void SetOwner(const TSharedPtr<IDecoderOutputOwner, ESPMode::ThreadSafe>& InOwningRenderer) override;
 	void ShutdownPoolable() override;
+	bool IsReadyForReuse() override;
 	EOutputType GetOutputType() const override;
-	TRefCountPtr<IMFSample> GetMFSample() const override;
 	const TArray<uint8>& GetBuffer() const override;
 	uint32 GetStride() const override;
 	TRefCountPtr<IUnknown> GetTexture() const override;
@@ -64,14 +73,21 @@ private:
 	TRefCountPtr<ID3D11Texture2D> SharedTexture;
 	TRefCountPtr<ID3D11Device> D3D11Device;
 
+	TRefCountPtr<ID3D12Resource> TextureDX12;
+	FIntPoint TextureDX12Dim = {0, 0};
+	TRefCountPtr<ID3D12Fence> D3DFence;
+	uint64 FenceValue = 0;
+
+	TRefCountPtr<ID3D12CommandAllocator> D3DCmdAllocator;
+	TRefCountPtr<ID3D12GraphicsCommandList> D3DCmdList;
+	TSharedPtr<FElectraMediaDecoderOutputBufferPool_DX12> D3D12ResourcePool;
+
+	TRefCountPtr<ID3D12Resource> DecoderOutputResource;
+
 	// CPU-side buffer
 	TSharedPtr<TArray<uint8>, ESPMode::ThreadSafe> Buffer;
 	uint32 Stride = 0;
 
-	// WMF sample (owned by this class if SW decoder is used)
-	TRefCountPtr<IMFSample> MFSample;
-	uint32 AllocatedBufferSize = 0;
-
 	// Dimension of any internally allocated buffer - stored explicitly to cover various special cases for DX
 	FIntPoint SampleDim;
 
diff --git a/Engine/Plugins/Media/ElectraPlayer/Source/ElectraPlayerRuntime/Private/Runtime/Decoder/Windows/VideoDecoderResourceDelegate.h b/Engine/Plugins/Media/ElectraPlayer/Source/ElectraPlayerRuntime/Private/Runtime/Decoder/Windows/VideoDecoderResourceDelegate.h
--- a/Engine/Plugins/Media/ElectraPlayer/Source/ElectraPlayerRuntime/Private/Runtime/Decoder/Windows/VideoDecoderResourceDelegate.h
+++ b/Engine/Plugins/Media/ElectraPlayer/Source/ElectraPlayerRuntime/Private/Runtime/Decoder/Windows/VideoDecoderResourceDelegate.h
@@ -2,12 +2,19 @@
 
 #pragma once
 
+#include "Templates/Function.h"
+#include "Templates/SharedPointer.h"
+
+struct ID3D12CommandQueue;
+
 namespace Electra {
 
-class IVideoDecoderResourceDelegate
+class IVideoDecoderResourceDelegate : public TSharedFromThis<IVideoDecoderResourceDelegate, ESPMode::ThreadSafe>
 {
 public:
 	virtual ~IVideoDecoderResourceDelegate() {}
+
+	virtual void ExecuteCodeWithCopyCommandQueueUsage(TFunction<void(ID3D12CommandQueue*)>&& CodeToRun) = 0;
 };
 
 }
diff --git a/Engine/Plugins/Media/ElectraPlayer/Source/ElectraPlayerRuntime/Private/Runtime/Decoder/Windows/WindowsElectraDecoderResourceManager.cpp b/Engine/Plugins/Media/ElectraPlayer/Source/ElectraPlayerRuntime/Private/Runtime/Decoder/Windows/WindowsElectraDecoderResourceManager.cpp
--- a/Engine/Plugins/Media/ElectraPlayer/Source/ElectraPlayerRuntime/Private/Runtime/Decoder/Windows/WindowsElectraDecoderResourceManager.cpp
+++ b/Engine/Plugins/Media/ElectraPlayer/Source/ElectraPlayerRuntime/Private/Runtime/Decoder/Windows/WindowsElectraDecoderResourceManager.cpp
@@ -9,9 +9,12 @@
 
 #include "Renderer/RendererBase.h"
 #include "ElectraVideoDecoder_PC.h"
+#include "VideoDecoderResourceDelegate.h"
 
 #include "ElectraPlayerMisc.h"
 
+#include COMPILED_PLATFORM_HEADER(ElectraDecoderGPUBufferHelpers.h)
+
 /***************************************************************************************************************************************************/
 /***************************************************************************************************************************************************/
 /***************************************************************************************************************************************************/
@@ -76,19 +79,32 @@ class FElectraDecoderResourceManagerWindows::FInstanceVars : public IElectraDeco
 {
 public:
 	uint32 Codec4CC = 0;
+	uint32 MaxWidth = 0;
+	uint32 MaxHeight = 0;
+	uint32 MaxOutputBuffers = 0;
+	TWeakPtr<IVideoDecoderResourceDelegate, ESPMode::ThreadSafe> VideoDecoderResourceDelegate;
 };
 
 
-
 IElectraDecoderResourceDelegateWindows::IDecoderPlatformResource* FElectraDecoderResourceManagerWindows::CreatePlatformResource(void* InOwnerHandle, EDecoderPlatformResourceType InDecoderResourceType, const TMap<FString, FVariant> InOptions)
 {
-	FInstanceVars* Vars = new FInstanceVars;
-	check(Vars);
-	if (Vars)
+	IVideoDecoderResourceDelegate* Delegate = reinterpret_cast<IVideoDecoderResourceDelegate*>(ElectraDecodersUtil::GetVariantValueSafeU64(InOptions, TEXT("VideoResourceDelegate"), 0));
+	if (Delegate)
 	{
-		Vars->Codec4CC = (uint32)ElectraDecodersUtil::GetVariantValueSafeU64(InOptions, TEXT("codec_4cc"), 0);
+		TSharedPtr<IVideoDecoderResourceDelegate, ESPMode::ThreadSafe> VideoDecoderResourceDelegate = Delegate->AsShared();
+		FInstanceVars* Vars = new FInstanceVars();
+		check(Vars);
+		if (Vars)
+		{
+			Vars->VideoDecoderResourceDelegate = VideoDecoderResourceDelegate;
+			Vars->Codec4CC = (uint32)ElectraDecodersUtil::GetVariantValueSafeU64(InOptions, TEXT("codec_4cc"), 0);
+			Vars->MaxWidth = (uint32)ElectraDecodersUtil::GetVariantValueSafeU64(InOptions, TEXT("max_width"), 1920);
+			Vars->MaxHeight = (uint32)ElectraDecodersUtil::GetVariantValueSafeU64(InOptions, TEXT("max_height"), 1080);
+			Vars->MaxOutputBuffers = (uint32)ElectraDecodersUtil::GetVariantValueSafeU64(InOptions, TEXT("max_output_buffers"), 5);
+		}
+		return Vars;
 	}
-	return Vars;
+	return nullptr;
 }
 
 
@@ -113,12 +129,34 @@ bool FElectraDecoderResourceManagerWindows::GetD3DDevice(void **OutD3DDevice, in
 	return false;
 }
 
+IElectraDecoderResourceDelegateWindows::IAsyncConsecutiveTaskSync* FElectraDecoderResourceManagerWindows::CreateAsyncConsecutiveTaskSync()
+{
+	if (WindowsDecoderResources::Callbacks.CreateAsyncConsecutiveTaskSync)
+	{
+		return WindowsDecoderResources::Callbacks.CreateAsyncConsecutiveTaskSync();
+	}
+	return nullptr;
+}
+
+bool FElectraDecoderResourceManagerWindows::RunCodeAsync(TFunction<void()>&& CodeToRun, IAsyncConsecutiveTaskSync* TaskSync)
+{
+	if (WindowsDecoderResources::Callbacks.RunCodeAsync)
+	{
+		return WindowsDecoderResources::Callbacks.RunCodeAsync(MoveTemp(CodeToRun), TaskSync);
+	}
+	return false;
+}
+
 
 bool FElectraDecoderResourceManagerWindows::SetupRenderBufferFromDecoderOutputFromMFSample(IMediaRenderer::IBuffer* InOutBufferToSetup, TSharedPtr<FParamDict, ESPMode::ThreadSafe> InOutBufferPropertes, TSharedPtr<IElectraDecoderVideoOutput, ESPMode::ThreadSafe> InDecoderOutput, FElectraDecoderResourceManagerWindows::IDecoderPlatformResource* InPlatformSpecificResource)
 {
 	TSharedPtr<FElectraPlayerVideoDecoderOutputPC, ESPMode::ThreadSafe> DecoderOutput = InOutBufferToSetup->GetBufferProperties().GetValue(RenderOptionKeys::Texture).GetSharedPointer<FElectraPlayerVideoDecoderOutputPC>();
-	if (DecoderOutput.IsValid())
+	FInstanceVars* Vars = static_cast<FInstanceVars*>(InPlatformSpecificResource);
+
+	if (DecoderOutput.IsValid() && Vars != nullptr)
 	{
+		TSharedPtr<IVideoDecoderResourceDelegate, ESPMode::ThreadSafe> PinnedResourceDelegate = Vars->VideoDecoderResourceDelegate.Pin();
+
 		FElectraVideoDecoderOutputCropValues Crop = InDecoderOutput->GetCropValues();
 		InOutBufferPropertes->Set(IDecoderOutputOptionNames::Width, FVariantValue((int64)InDecoderOutput->GetWidth()));
 		InOutBufferPropertes->Set(IDecoderOutputOptionNames::Height, FVariantValue((int64)InDecoderOutput->GetHeight()));
@@ -174,6 +212,7 @@ bool FElectraDecoderResourceManagerWindows::SetupRenderBufferFromDecoderOutputFr
 				return false;
 			}
 
+			// No, continue with SW or DX11 texture output
 			TRefCountPtr<ID3D11Texture2D> Texture2D;
 			if ((Result = DXGIBuffer->GetResource(IID_PPV_ARGS(Texture2D.GetInitReference()))) != S_OK)
 			{
@@ -194,7 +233,7 @@ bool FElectraDecoderResourceManagerWindows::SetupRenderBufferFromDecoderOutputFr
 			if (DXVersion == 0 || DXVersion >= 12000)
 			{
 				//
-				// DX12 & non-DX
+				// DX12 (with DX11 decode device) & non-DX
 				// 
 				// (access buffer for CPU use)
 				//
@@ -441,6 +480,11 @@ bool FElectraDecoderResourceManagerWindows::SetupRenderBufferFromDecoderOutput(I
 			InOutBufferPropertes->Set(IDecoderOutputOptionNames::PixelEncoding, FVariantValue((int64)DecPixEnc));
 			InOutBufferPropertes->Set(IDecoderOutputOptionNames::Pitch, FVariantValue((int64)Pitch));
 
+			TMap<FString, FVariant> ExtraValues;
+			InDecoderOutput->GetExtraValues(ExtraValues);
+
+			InOutBufferPropertes->Set(IDecoderOutputOptionNames::PixelDataScale, FVariantValue((double)ElectraDecodersUtil::GetVariantValueSafeDouble(ExtraValues, TEXT("pix_datascale"), 1.0)));
+
 // [...] ALPHA BUFFER -- HOW DO WE PASS IT ON!?
 // [...] ANY CS/HDR INFO FROM THE DECODER? -- PASSING IT ON WOULD BE EASY, BUT RIGHT NOW WE ALWAYS READ IT FROM THE CONTAINER (in code further up the chain)
 
@@ -458,24 +502,57 @@ bool FElectraDecoderResourceManagerWindows::SetupRenderBufferFromDecoderOutput(I
 			}
 			else
 			{
-				// Note: we assume a DX11 texture at this point, but this could also be interpreted as IUnknown and then either DX11 or DX12 resources being derived from it...
-				ID3D11Texture2D* Texture = static_cast<ID3D11Texture2D*>(ImageBuffers->GetBufferTextureByIndex(0));
-				if (Texture != nullptr)
+				TRefCountPtr<IUnknown> TextureCommon(static_cast<IUnknown*>(ImageBuffers->GetBufferTextureByIndex(0)));
+
+				if (TextureCommon.IsValid())
 				{
-					//
-					// GPU texture
-					//
-					ID3D11Device* Device = reinterpret_cast<ID3D11Device*>(InDecoderOutput->GetPlatformOutputHandle(EElectraDecoderPlatformOutputHandleType::DXDevice));
-					ID3D11DeviceContext* DeviceContext = reinterpret_cast<ID3D11DeviceContext*>(InDecoderOutput->GetPlatformOutputHandle(EElectraDecoderPlatformOutputHandleType::DXDeviceContext));
-					if (Device && DeviceContext)
+					HRESULT Res;
+
+					void* DDevice;
+					int32 DVersion;
+					check(WindowsDecoderResources::Callbacks.GetD3DDevice);
+					WindowsDecoderResources::Callbacks.GetD3DDevice(&DDevice, &DVersion, WindowsDecoderResources::Callbacks.UserValue);
+
+					if (DVersion >= 12000)
+					{
+						TRefCountPtr<ID3D12Device> D3D12Device(static_cast<ID3D12Device*>(DDevice));
+
+						// Can we get a DX12 texture? (this will only work if the transform is associated with a DX12 device earlier on; hence no need for a SDK version guard here)
+						TRefCountPtr<ID3D12Resource> Resource;
+						Res = TextureCommon->QueryInterface(__uuidof(ID3D12Resource), (void**)Resource.GetInitReference());
+						if (Res == S_OK)
+						{
+							//
+							// DX12 texture / buffer
+							//
+
+							// We might also have a fence / sync we need to use before accessing the data on the texture...
+							FElectraDecoderOutputSync OutputSync;
+							ImageBuffers->GetBufferTextureSyncByIndex(0, OutputSync);
+
+							DecoderOutput->InitializeWithResource(D3D12Device, Resource, Pitch, OutputSync, FIntPoint(InDecoderOutput->GetDecodedWidth(), InDecoderOutput->GetDecodedHeight()), InOutBufferPropertes, Vars->VideoDecoderResourceDelegate.Pin().Get(), Vars->MaxWidth, Vars->MaxHeight, Vars->MaxOutputBuffers);
+							return true;
+						}
+					}
+
+					TRefCountPtr<ID3D11Texture2D> Texture;
+					Res = TextureCommon->QueryInterface(__uuidof(ID3D11Texture2D), (void**)Texture.GetInitReference());
+					if (Res == S_OK)
 					{
-						uint32 ViewIndex = 0;
-						DecoderOutput->InitializeWithSharedTexture(Device, DeviceContext, Texture, ViewIndex, FIntPoint(InDecoderOutput->GetWidth(), InDecoderOutput->GetHeight()), InOutBufferPropertes);
-						return true;
+						//
+						// DX11 texture
+						//
+						ID3D11Device* Device = reinterpret_cast<ID3D11Device*>(InDecoderOutput->GetPlatformOutputHandle(EElectraDecoderPlatformOutputHandleType::DXDevice));
+						ID3D11DeviceContext* DeviceContext = reinterpret_cast<ID3D11DeviceContext*>(InDecoderOutput->GetPlatformOutputHandle(EElectraDecoderPlatformOutputHandleType::DXDeviceContext));
+						if (Device && DeviceContext)
+						{
+							uint32 ViewIndex = 0;
+							DecoderOutput->InitializeWithSharedTexture(Device, DeviceContext, Texture.GetReference(), ViewIndex, FIntPoint(InDecoderOutput->GetWidth(), InDecoderOutput->GetHeight()), InOutBufferPropertes);
+							return true;
+						}
 					}
 				}
 			}
-
 		}
 		else
 		{
@@ -492,7 +569,5 @@ bool FElectraDecoderResourceManagerWindows::SetupRenderBufferFromDecoderOutput(I
 }
 
 
-
-
 } // namespace Electra
 
diff --git a/Engine/Plugins/Media/ElectraPlayer/Source/ElectraPlayerRuntime/Private/Runtime/Decoder/Windows/WindowsElectraDecoderResourceManager.h b/Engine/Plugins/Media/ElectraPlayer/Source/ElectraPlayerRuntime/Private/Runtime/Decoder/Windows/WindowsElectraDecoderResourceManager.h
--- a/Engine/Plugins/Media/ElectraPlayer/Source/ElectraPlayerRuntime/Private/Runtime/Decoder/Windows/WindowsElectraDecoderResourceManager.h
+++ b/Engine/Plugins/Media/ElectraPlayer/Source/ElectraPlayerRuntime/Private/Runtime/Decoder/Windows/WindowsElectraDecoderResourceManager.h
@@ -18,6 +18,9 @@ public:
 	struct FCallbacks
 	{
 		TFunction<bool(void ** /*OutD3DDevice*/, int32* /*OutD3DVersionTimes1000*/, void* /*UserValue*/)> GetD3DDevice;
+		TFunction<bool(TFunction<void()>&& CodeToRun, IAsyncConsecutiveTaskSync* TaskSync)> RunCodeAsync;
+		TFunction<IAsyncConsecutiveTaskSync* ()> CreateAsyncConsecutiveTaskSync;
+
 		void* UserValue = nullptr;
 	};
 
@@ -32,6 +35,9 @@ public:
 
 	virtual bool GetD3DDevice(void **OutD3DDevice, int32* OutD3DVersionTimes1000) override;
 
+	virtual IAsyncConsecutiveTaskSync* CreateAsyncConsecutiveTaskSync() override;
+	virtual bool RunCodeAsync(TFunction<void()>&& CodeToRun, IAsyncConsecutiveTaskSync* TaskSync) override;
+
 	virtual ~FElectraDecoderResourceManagerWindows();
 
 private:
diff --git a/Engine/Plugins/Media/ElectraPlayer/Source/ElectraPlayerRuntime/Private/Windows/ElectraPlayer_Windows.cpp b/Engine/Plugins/Media/ElectraPlayer/Source/ElectraPlayerRuntime/Private/Windows/ElectraPlayer_Windows.cpp
--- a/Engine/Plugins/Media/ElectraPlayer/Source/ElectraPlayerRuntime/Private/Windows/ElectraPlayer_Windows.cpp
+++ b/Engine/Plugins/Media/ElectraPlayer/Source/ElectraPlayerRuntime/Private/Windows/ElectraPlayer_Windows.cpp
@@ -10,6 +10,8 @@
 #include "ElectraPlayer.h"
 
 #include "ParameterDictionary.h"
+#include "VideoDecoderResourceDelegate.h"
+
 #include "RHIDefinitions.h"
 
 #include "WindowsElectraDecoderResourceManager.h"
@@ -33,21 +35,43 @@ THIRD_PARTY_INCLUDES_END
 namespace Electra
 {
 
-// ----------------------------------------------------------------------------------------------------------------------
+	// ----------------------------------------------------------------------------------------------------------------------
 
-/**
- * Shutdown of module
- */
-bool PlatformShutdown()
-{
-	FElectraDecoderResourceManagerWindows::Shutdown();
-	return true;
-}
+	/**
+	 * Shutdown of module
+	 */
+	bool PlatformShutdown()
+	{
+		FElectraDecoderResourceManagerWindows::Shutdown();
+		return true;
+	}
 
-TSharedPtr<IVideoDecoderResourceDelegate, ESPMode::ThreadSafe> PlatformCreateVideoDecoderResourceDelegate(const TWeakPtr<IElectraPlayerAdapterDelegate, ESPMode::ThreadSafe>& AdapterDelegate)
-{
-	return nullptr;
-}
+	class FVideoDecoderResourceDelegate : public IVideoDecoderResourceDelegate
+	{
+	public:
+		FVideoDecoderResourceDelegate(const TWeakPtr<IElectraPlayerAdapterDelegate, ESPMode::ThreadSafe>& InAdapterDelegate) : PlayerAdapterDelegate(InAdapterDelegate) {}
+		virtual ~FVideoDecoderResourceDelegate() {}
+
+		virtual void ExecuteCodeWithCopyCommandQueueUsage(TFunction<void(ID3D12CommandQueue*)>&& CodeToRun) override
+		{
+			if (const TSharedPtr<IElectraPlayerAdapterDelegate, ESPMode::ThreadSafe>& PinnedAdapterDelegate = PlayerAdapterDelegate.Pin())
+			{
+				TSharedPtr<IElectraPlayerResourceDelegate, ESPMode::ThreadSafe> ResourceDelegate = PinnedAdapterDelegate->GetResourceDelegate();
+				if (ResourceDelegate)
+				{
+					ResourceDelegate->ExecuteCodeWithCopyCommandQueueUsage(MoveTemp(CodeToRun));
+				}
+			}
+		}
+
+	private:
+		TWeakPtr<IElectraPlayerAdapterDelegate, ESPMode::ThreadSafe> PlayerAdapterDelegate;
+	};
+
+	TSharedPtr<IVideoDecoderResourceDelegate, ESPMode::ThreadSafe> PlatformCreateVideoDecoderResourceDelegate(const TWeakPtr<IElectraPlayerAdapterDelegate, ESPMode::ThreadSafe>& AdapterDelegate)
+	{
+		return MakeShared<FVideoDecoderResourceDelegate, ESPMode::ThreadSafe>(AdapterDelegate);
+	}
 
 } //namespace Electra
 
@@ -62,23 +86,47 @@ bool FElectraPlayerPlatform::StartupPlatformResources(const Electra::FParamDict&
 	FElectraDecoderResourceManagerWindows::FCallbacks Callbacks;
 	auto GetDeviceTypeCallback = reinterpret_cast<void(*)(void**, int64*)>(Params.GetValue(FName(TEXT("GetDeviceTypeCallback"))).SafeGetPointer());
 
-	Callbacks.GetD3DDevice = [GetDeviceTypeCallback](void **OutD3DDevice, int32* OutD3DVersionTimes1000, void*) -> bool
-	{
-		if (GetDeviceTypeCallback)
+	Callbacks.GetD3DDevice = [GetDeviceTypeCallback](void** OutD3DDevice, int32* OutD3DVersionTimes1000, void*) -> bool
 		{
-			void* DevicePointer = nullptr;
-			int64 DeviceType = 0;
-			GetDeviceTypeCallback(&DevicePointer, &DeviceType);
-			const ERHIInterfaceType RHIType = (ERHIInterfaceType)DeviceType;
-			if (DevicePointer && (RHIType == ERHIInterfaceType::D3D11 || RHIType == ERHIInterfaceType::D3D12))
+			if (GetDeviceTypeCallback)
 			{
-				*OutD3DDevice = DevicePointer;
-				*OutD3DVersionTimes1000 = RHIType == ERHIInterfaceType::D3D11 ? 11000 : 12000;
-				return true;
+				void* DevicePointer = nullptr;
+				int64 DeviceType = 0;
+				GetDeviceTypeCallback(&DevicePointer, &DeviceType);
+				const ERHIInterfaceType RHIType = (ERHIInterfaceType)DeviceType;
+				if (DevicePointer && (RHIType == ERHIInterfaceType::D3D11 || RHIType == ERHIInterfaceType::D3D12))
+				{
+					*OutD3DDevice = DevicePointer;
+					*OutD3DVersionTimes1000 = RHIType == ERHIInterfaceType::D3D11 ? 11000 : 12000;
+					return true;
+				}
 			}
+			return false;
+		};
+
+	auto RunCodeAsyncCallback = reinterpret_cast<void(*)(TFunction<void()>&& CodeToRun, IElectraDecoderResourceDelegateWindows::IAsyncConsecutiveTaskSync * TaskSync)>(Params.GetValue(FName(TEXT("RunCodeAsyncCallback"))).SafeGetPointer());
+
+	Callbacks.RunCodeAsync = [RunCodeAsyncCallback](TFunction<void()>&& CodeToRun, IElectraDecoderResourceDelegateWindows::IAsyncConsecutiveTaskSync* TaskSync) -> bool
+	{
+		if (RunCodeAsyncCallback)
+		{
+			RunCodeAsyncCallback(MoveTemp(CodeToRun), TaskSync);
+			return true;
 		}
 		return false;
 	};
+
+	auto CreateAsyncConsecutiveTaskSyncCallback = reinterpret_cast<IElectraDecoderResourceDelegateWindows::IAsyncConsecutiveTaskSync*(*)()>(Params.GetValue(FName(TEXT("CreateAsyncConsecutiveTaskSync"))).SafeGetPointer());
+
+	Callbacks.CreateAsyncConsecutiveTaskSync = [CreateAsyncConsecutiveTaskSyncCallback]() -> IElectraDecoderResourceDelegateWindows::IAsyncConsecutiveTaskSync*
+	{
+		if (CreateAsyncConsecutiveTaskSyncCallback)
+		{
+			return CreateAsyncConsecutiveTaskSyncCallback();
+		}
+		return nullptr;
+	};
+
 	return FElectraDecoderResourceManagerWindows::Startup(Callbacks);
 }
 
diff --git a/Engine/Plugins/Media/ElectraPlayer/Source/ElectraPlayerRuntime/Public/Windows/ElectraPlayerResourceDelegate.h b/Engine/Plugins/Media/ElectraPlayer/Source/ElectraPlayerRuntime/Public/Windows/ElectraPlayerResourceDelegate.h
--- a/Engine/Plugins/Media/ElectraPlayer/Source/ElectraPlayerRuntime/Public/Windows/ElectraPlayerResourceDelegate.h
+++ b/Engine/Plugins/Media/ElectraPlayer/Source/ElectraPlayerRuntime/Public/Windows/ElectraPlayerResourceDelegate.h
@@ -1,8 +1,12 @@
 // Copyright Epic Games, Inc. All Rights Reserved.
 #pragma once
 
+struct ID3D12CommandQueue;
+
 class IElectraPlayerResourceDelegate
 {
 public:
 	virtual ~IElectraPlayerResourceDelegate() {}
+
+	virtual void ExecuteCodeWithCopyCommandQueueUsage(TFunction<void(ID3D12CommandQueue*)>&& CodeToRun) = 0;
 };
diff --git a/Engine/Plugins/Media/ElectraUtil/Source/ElectraBase/Public/MediaDecoderOutput.h b/Engine/Plugins/Media/ElectraUtil/Source/ElectraBase/Public/MediaDecoderOutput.h
--- a/Engine/Plugins/Media/ElectraUtil/Source/ElectraBase/Public/MediaDecoderOutput.h
+++ b/Engine/Plugins/Media/ElectraUtil/Source/ElectraBase/Public/MediaDecoderOutput.h
@@ -7,10 +7,10 @@
 #include "Templates/RefCounting.h"
 #include "Containers/Queue.h"
 #include "Misc/ScopeLock.h"
+#include "HAL/PlatformProcess.h"
 
 #include "ParameterDictionary.h"
 
-
 struct FDecoderTimeStamp
 {
 	FDecoderTimeStamp() {}
@@ -47,6 +47,8 @@ class TDecoderOutputObjectPool
 	/** Object pool storage. */
 	class TStorage
 	{
+		static float constexpr kStorageBusyWaitPollInterval = 0.002f;
+
 	public:
 		TStorage(ObjectFactory* InObjectFactoryInstance)
 			: ObjectFactoryInstance(InObjectFactoryInstance)
@@ -56,12 +58,31 @@ class TDecoderOutputObjectPool
 		~TStorage()
 		{
 			Reserve(0);
+
+			double StartTime = FPlatformTime::Seconds();
+			bool bDidWait = false;
 			ObjectType* Object;
 			while (WaitReadyForReuse.Dequeue(Object))
 			{
+				while (!Object->IsReadyForReuse() && GIsRunning)
+				{
+					// If we encounter an object that is still busy, we wait a bit and try again
+					// (Assumptions:
+					//  - there will not be that many busy objects
+					//  - usually the state of a busy object should become idle quite quickly)
+					bDidWait = true;
+					FPlatformProcess::Sleep(kStorageBusyWaitPollInterval);
+				}
+
 				Object->ShutdownPoolable();
 				delete Object;
 			}
+
+			if (bDidWait)
+			{
+				double EndTime = FPlatformTime::Seconds();
+				UE_LOG(LogTemp, Log, TEXT("[%p] TDecoderOutputObjectPool::TStorage::~TStorage() finished after %.3f msec!"), this, (EndTime - StartTime) * 1000.0);
+			}
 		}
 
 	public:
@@ -324,4 +345,5 @@ static const FName AspectW(TEXT("aspect_w"));
 static const FName AspectH(TEXT("aspect_h"));
 static const FName FPSNumerator(TEXT("fps_num"));
 static const FName FPSDenominator(TEXT("fps_denom"));
+static const FName PixelDataScale(TEXT("pix_datascale"));
 }
diff --git a/Engine/Plugins/Media/ElectraUtil/Source/ElectraBase/Public/MediaVideoDecoderOutputPC.h b/Engine/Plugins/Media/ElectraUtil/Source/ElectraBase/Public/MediaVideoDecoderOutputPC.h
--- a/Engine/Plugins/Media/ElectraUtil/Source/ElectraBase/Public/MediaVideoDecoderOutputPC.h
+++ b/Engine/Plugins/Media/ElectraUtil/Source/ElectraBase/Public/MediaVideoDecoderOutputPC.h
@@ -17,6 +17,7 @@
 	#pragma warning(disable : 4005) // macro redefinition
 	#include <d3d11.h>
 	#pragma warning(pop)
+	#include <d3d12.h>
 
 THIRD_PARTY_INCLUDES_START
 	#include "mfobjects.h"
@@ -43,7 +44,6 @@ public:
 	virtual uint32 GetStride() const = 0;
 
 #if PLATFORM_WINDOWS
-	virtual TRefCountPtr<IMFSample> GetMFSample() const = 0;
 	virtual TRefCountPtr<IUnknown> GetTexture() const = 0;
 	virtual TRefCountPtr<IUnknown> GetSync(uint64& SyncValue) const = 0;
 	virtual TRefCountPtr<ID3D11Device> GetDevice() const = 0;
diff --git a/Engine/Plugins/Media/ElectraUtil/Source/ElectraSamples/ElectraSamples.Build.cs b/Engine/Plugins/Media/ElectraUtil/Source/ElectraSamples/ElectraSamples.Build.cs
--- a/Engine/Plugins/Media/ElectraUtil/Source/ElectraSamples/ElectraSamples.Build.cs
+++ b/Engine/Plugins/Media/ElectraUtil/Source/ElectraSamples/ElectraSamples.Build.cs
@@ -39,23 +39,23 @@ namespace UnrealBuildTool.Rules
 			if (Target.Platform.IsInGroup(UnrealPlatformGroup.Windows))
 			{
 				PublicDependencyModuleNames.Add("DirectX");
-				if (Target.Platform.IsInGroup(UnrealPlatformGroup.Windows))
-				{
-					AddEngineThirdPartyPrivateStaticDependencies(Target, "DX9");
-
-					if (Target.WindowsPlatform.Architecture != UnrealArch.Arm64)
-					{
-						PublicAdditionalLibraries.AddRange(new string[] {
-							Path.Combine(Target.WindowsPlatform.DirectXLibDir, "dxerr.lib"),
-						});
-					}
 
-					PrivateDependencyModuleNames.Add("D3D11RHI");
+				AddEngineThirdPartyPrivateStaticDependencies(Target, "DX9");
 
-					PrivateDefinitions.Add("ELECTRA_SUPPORT_PREWIN8");
-					PrivateDefinitions.Add("ELECTRA_HAVE_DX11");
+				if (Target.WindowsPlatform.Architecture != UnrealArch.Arm64)
+				{
+					PublicAdditionalLibraries.AddRange(new string[] {
+						Path.Combine(Target.WindowsPlatform.DirectXLibDir, "dxerr.lib"),
+					});
 				}
 
+				PrivateDependencyModuleNames.Add("D3D11RHI");
+				PrivateDependencyModuleNames.Add("D3D12RHI");
+				AddEngineThirdPartyPrivateStaticDependencies(Target, "DX12");
+
+				PrivateDefinitions.Add("ELECTRA_SUPPORT_PREWIN8");
+				PrivateDefinitions.Add("ELECTRA_HAVE_DX11");
+
 				PublicSystemLibraries.AddRange(new string[] {
 					"strmiids.lib",
 					"legacy_stdio_definitions.lib",
diff --git a/Engine/Plugins/Media/ElectraUtil/Source/ElectraSamples/Private/ElectraTextureSamplePC.cpp b/Engine/Plugins/Media/ElectraUtil/Source/ElectraSamples/Private/ElectraTextureSamplePC.cpp
--- a/Engine/Plugins/Media/ElectraUtil/Source/ElectraSamples/Private/ElectraTextureSamplePC.cpp
+++ b/Engine/Plugins/Media/ElectraUtil/Source/ElectraSamples/Private/ElectraTextureSamplePC.cpp
@@ -9,20 +9,25 @@
 #include "ProfilingDebugging/RealtimeGPUProfiler.h"
 #include "RenderUtils.h"
 
+#include "ID3D12DynamicRHI.h"
+
 #include "Windows/AllowWindowsPlatformTypes.h"
+THIRD_PARTY_INCLUDES_START
 #ifdef ELECTRA_HAVE_DX11
 #include "D3D11State.h"
 #include "D3D11Resources.h"
 #endif
 #include "d3d12.h"
+THIRD_PARTY_INCLUDES_END
 #include "Windows/HideWindowsPlatformTypes.h"
 
+
 /*
 	Short summary of how we get data:
 
 	- Win10+ (HW decode is used at all times if handling H.264/5)
 	-- DX11:   we receive data in GPU space as NV12/P010 texture
-	-- DX12:   we receive data in CPU(yes) space as NV12/P010 texture
+	-- DX12:   we receive data in CPU(yes) space as NV12/P010 texture -=UNLESS=- we have SDK 22621+ and suitable runtime support -> DX12 texture in GPU space
 	-- Vulkan: we receive data in CPU(yes) space as NV12/P010 texture
 	-- Other codec's data usually arrives as CPU space texture buffer
 
@@ -41,8 +46,8 @@ DECLARE_GPU_STAT_NAMED(MediaWinDecoder_Convert, TEXT("MediaWinDecoder_Convert"))
 
 void FElectraTextureSample::Initialize(FVideoDecoderOutput *InVideoDecoderOutput)
 {
-	IElectraTextureSampleBase::Initialize(InVideoDecoderOutput);
 	VideoDecoderOutputPC = static_cast<FVideoDecoderOutputPC*>(InVideoDecoderOutput);
+	IElectraTextureSampleBase::Initialize(InVideoDecoderOutput);
 
 	EPixelFormat Format = VideoDecoderOutput->GetFormat();
 	switch (Format)
@@ -156,25 +161,73 @@ void FElectraTextureSample::Initialize(FVideoDecoderOutput *InVideoDecoderOutput
 		check(!"Decoder sample format not supported in Electra texture sample!");
 	}
 
-	bCanUseSRGB = (Format == PF_B8G8R8A8 || Format == PF_R8G8B8A8 || Format == PF_DXT1 || Format == PF_DXT5 || Format == PF_BC4);
+	EMediaTextureSampleFormat SampleFmt = GetFormat();
+	bCanUseSRGB = (SampleFmt == EMediaTextureSampleFormat::CharBGRA ||
+				   SampleFmt == EMediaTextureSampleFormat::CharRGBA ||
+				   SampleFmt == EMediaTextureSampleFormat::CharBMP ||
+				   SampleFmt == EMediaTextureSampleFormat::DXT1 ||
+				   SampleFmt == EMediaTextureSampleFormat::DXT5);
+
+	// Get rid of older textures if we switched around to CPU side buffers (which might happen if a single player switches between decoders)
+	// (note: we do not care about the CPU side buffers as we do not keep a reference on that resource)
+	if (Texture.IsValid() && (VideoDecoderOutputPC->GetTexture() == nullptr))
+	{
+		Texture = nullptr;
+	}
+}
+
+
+#if !UE_SERVER
+void FElectraTextureSample::ShutdownPoolable()
+{
+	// If we use DX12 with texture data, we use the RHI texture only to refer to the actual "native" D3D resource we created -> get rid of this "alias" as soon as we retire to the pool
+//OPT: WE COULD TRACK IF WE STILL HAVE THE SAME NATIVE TEXTURE AND ATTEMPT TO REUSE THIS
+	if (VideoDecoderOutputPC->GetTexture() && (RHIGetInterfaceType() == ERHIInterfaceType::D3D12))
+	{
+		Texture = nullptr;
+	}
+	IElectraTextureSampleBase::ShutdownPoolable();
+}
+#endif
+
 
+#if WITH_ENGINE
+FRHITexture* FElectraTextureSample::GetTexture() const
+{
+	check(IsInRenderingThread());
+
+	// Is this DX12?
 	if (RHIGetInterfaceType() == ERHIInterfaceType::D3D12)
 	{
-		ID3D12Device* ApplicationDxDevice = static_cast<ID3D12Device*>(GDynamicRHI->RHIGetNativeDevice());
-		HRESULT Res;
-		if (!D3DCmdAllocator.IsValid())
+		// Yes, see if we have a real D3D resource...
+		if (TRefCountPtr<IUnknown> TextureCommon = VideoDecoderOutputPC->GetTexture())
 		{
-			Res = ApplicationDxDevice->CreateCommandAllocator(D3D12_COMMAND_LIST_TYPE_DIRECT, IID_PPV_ARGS(D3DCmdAllocator.GetInitReference()));
-			check(!FAILED(Res));
-		}
-		if (!D3DCmdList.IsValid())
-		{
-			Res = ApplicationDxDevice->CreateCommandList(0, D3D12_COMMAND_LIST_TYPE_DIRECT, D3DCmdAllocator.GetReference(), nullptr, __uuidof(ID3D12CommandList), reinterpret_cast<void**>(D3DCmdList.GetInitReference()));
-			check(!FAILED(Res));
-			D3DCmdList->Close();
+			// Yes. Do we need a new RHI texture?
+			FIntPoint Dim = VideoDecoderOutput->GetDim();
+			if (!Texture.IsValid() || Texture->GetSizeX() != Dim.X || Texture->GetSizeY() != Dim.Y)
+			{
+				// Yes...
+				TRefCountPtr<ID3D12Resource> TextureDX12;
+				HRESULT Res = TextureCommon->QueryInterface(__uuidof(ID3D12Resource), (void**)TextureDX12.GetInitReference());
+				check(Res == S_OK);
+				if (Res == S_OK)
+				{
+					// Setup a suitable RHI texture (it will also become an additional owner of the data)
+					ETextureCreateFlags Flags = ETextureCreateFlags::ShaderResource;
+					if (IsOutputSrgb() && bCanUseSRGB)
+					{
+						Flags = Flags | ETextureCreateFlags::SRGB;
+					}
+					Texture = static_cast<ID3D12DynamicRHI*>(GDynamicRHI)->RHICreateTexture2DFromResource(VideoDecoderOutput->GetFormat(), Flags, FClearValueBinding(), TextureDX12);
+				}
+			}
 		}
 	}
+
+	return Texture;
 }
+#endif //WITH_ENGINE
+
 
 IMediaTextureSampleConverter* FElectraTextureSample::GetMediaTextureSampleConverter()
 {
@@ -195,70 +248,6 @@ IMediaTextureSampleConverter* FElectraTextureSample::GetMediaTextureSampleConver
 }
 
 
-struct FRHICommandCopyResourceDX12 final : public FRHICommand<FRHICommandCopyResourceDX12>
-{
-	TRefCountPtr<ID3D12Resource> SampleTexture;
-	FTexture2DRHIRef SampleDestinationTexture;
-	TRefCountPtr<ID3D12CommandAllocator> D3DCmdAllocator;
-	TRefCountPtr<ID3D12GraphicsCommandList> D3DCmdList;
-	TRefCountPtr<ID3D12Fence> D3DFence;
-	uint64 FenceValue;
-
-	FRHICommandCopyResourceDX12(ID3D12Resource* InSampleTexture, FRHITexture2D* InSampleDestinationTexture, TRefCountPtr<ID3D12CommandAllocator> InD3DCmdAllocator, TRefCountPtr<ID3D12GraphicsCommandList> InD3DCmdList, TRefCountPtr<ID3D12Fence> InD3DFence, uint64 InFenceValue)
-		: SampleTexture(InSampleTexture)
-		, SampleDestinationTexture(InSampleDestinationTexture)
-		, D3DCmdAllocator(InD3DCmdAllocator)
-		, D3DCmdList(InD3DCmdList)
-		, D3DFence(InD3DFence)
-		, FenceValue(InFenceValue)
-	{
-		LLM_SCOPE(ELLMTag::MediaStreaming);
-
-		auto DstTexture = (ID3D12Resource*)SampleDestinationTexture->GetNativeResource();
-
-		D3D12_RESOURCE_BARRIER PreCopyResourceBarriers[] = {
-			{D3D12_RESOURCE_BARRIER_TYPE_TRANSITION, D3D12_RESOURCE_BARRIER_FLAG_NONE, {{ SampleTexture, D3D12_RESOURCE_BARRIER_ALL_SUBRESOURCES, D3D12_RESOURCE_STATE_COMMON, D3D12_RESOURCE_STATE_COPY_SOURCE }}},
-			{D3D12_RESOURCE_BARRIER_TYPE_TRANSITION, D3D12_RESOURCE_BARRIER_FLAG_NONE, {{ DstTexture, D3D12_RESOURCE_BARRIER_ALL_SUBRESOURCES, D3D12_RESOURCE_STATE_NON_PIXEL_SHADER_RESOURCE | D3D12_RESOURCE_STATE_PIXEL_SHADER_RESOURCE, D3D12_RESOURCE_STATE_COPY_DEST }}},
-		};
-
-		D3D12_RESOURCE_BARRIER PostCopyResourceBarriers[] = {
-			{ D3D12_RESOURCE_BARRIER_TYPE_TRANSITION, D3D12_RESOURCE_BARRIER_FLAG_NONE, {{ SampleTexture, D3D12_RESOURCE_BARRIER_ALL_SUBRESOURCES, D3D12_RESOURCE_STATE_COPY_SOURCE, D3D12_RESOURCE_STATE_COMMON }}},
-			{ D3D12_RESOURCE_BARRIER_TYPE_TRANSITION, D3D12_RESOURCE_BARRIER_FLAG_NONE, {{ DstTexture, D3D12_RESOURCE_BARRIER_ALL_SUBRESOURCES, D3D12_RESOURCE_STATE_COPY_DEST, D3D12_RESOURCE_STATE_NON_PIXEL_SHADER_RESOURCE | D3D12_RESOURCE_STATE_PIXEL_SHADER_RESOURCE }}},
-		};
-
-		D3DCmdAllocator->Reset();
-		D3DCmdList->Reset(D3DCmdAllocator, nullptr);
-
-		D3DCmdList->ResourceBarrier(2, PreCopyResourceBarriers);
-		D3DCmdList->CopyResource(DstTexture, SampleTexture);
-		D3DCmdList->ResourceBarrier(2, PostCopyResourceBarriers);
-
-		D3DCmdList->Close();
-	}
-
-	void Execute(FRHICommandListBase& CmdList)
-	{
-		TRefCountPtr<ID3D12GraphicsCommandList> D3DCmdList2 = D3DCmdList;
-		TRefCountPtr<ID3D12CommandAllocator> D3DCmdAllocator2 = D3DCmdAllocator;
-		TRefCountPtr<ID3D12Resource> SampleTexture2 = SampleTexture;
-		TRefCountPtr<ID3D12Fence> D3DFence2 = D3DFence;
-		FTexture2DRHIRef SampleDestinationTexture2 = SampleDestinationTexture;
-		ENQUEUE_RENDER_COMMAND(MediaCopyInputTexture)(
-			[D3DCmdList2, D3DCmdAllocator2, SampleTexture2, SampleDestinationTexture2, D3DFence2, FenceValue2 = FenceValue](FRHICommandListImmediate& RHICmdList)
-		{
-			LLM_SCOPE(ELLMTag::MediaStreaming);
-			auto D3DCmdQueue = (ID3D12CommandQueue*)FRHICommandListExecutor::GetImmediateCommandList().GetNativeGraphicsQueue();
-			if (D3DFence2.IsValid())
-			{
-				D3DCmdQueue->Wait(D3DFence2, FenceValue2);
-			}
-			ID3D12CommandList* CmdLists[1] = {D3DCmdList2.GetReference()};
-			D3DCmdQueue->ExecuteCommandLists(1, CmdLists);
-		});
-	}
-};
-
-
 #ifdef ELECTRA_HAVE_DX11
 struct FRHICommandCopyResourceDX11 final : public FRHICommand<FRHICommandCopyResourceDX11>
 {
@@ -363,6 +352,9 @@ bool FElectraTextureSample::Convert(FTexture2DRHIRef& InDstTexture, const FConve
 	SCOPED_DRAW_EVENT(RHICmdList, WinMediaOutputConvertTexture);
 	SCOPED_GPU_STAT(RHICmdList, MediaWinDecoder_Convert);
 
+
+	bool bHasTexture = !!VideoDecoderOutputPC->GetTexture();
+
 	// Get actual sample dimensions
 	FIntPoint Dim = VideoDecoderOutput->GetDim();
 
@@ -374,12 +366,41 @@ bool FElectraTextureSample::Convert(FTexture2DRHIRef& InDstTexture, const FConve
 		return false;
 	}
 
-	// Note: the converter is not used at all if we have texture data in a SW buffer!
+	// If we have a texture and D3D12 is active, we will not need to do any custom conversion work, but we will need a sync!
+	if (bHasTexture && (RHIGetInterfaceType() == ERHIInterfaceType::D3D12))
+	{
+		uint64 SyncFenceValue;
+		TRefCountPtr<IUnknown> SyncCommon = VideoDecoderOutputPC->GetSync(SyncFenceValue);
+		if (SyncCommon)
+		{
+			//
+			// Synchronize with sample data copy on copy queue
+			//
+			TRefCountPtr<ID3D12Fence> SyncFence;
+			HRESULT Res = SyncCommon->QueryInterface(__uuidof(ID3D12Fence), (void**)SyncFence.GetInitReference());
+			check(SUCCEEDED(Res));
+			if (Res == S_OK)
+			{
+				RHICmdList.EnqueueLambda([SyncFence, SyncFenceValue](FRHICommandList& RHICmdList)
+				{
+					GetID3D12DynamicRHI()->RHIWaitManualFence(RHICmdList, SyncFence, SyncFenceValue);
+				});
+			}
+		}
+		return true;
+	}
+
+	// Note: the converter code (from here on) is not used at all if we have texture data in a SW buffer!
+	check(bHasTexture);
 
 	bool bCrossDeviceCopy;
 	EPixelFormat Format;
 	if (VideoDecoderOutputPC->GetOutputType() != FVideoDecoderOutputPC::EOutputType::HardwareWin8Plus)
 	{
+		// Below case deliver DX11 textures on the rendering device, so it would be feasible to create an RHI texture from them directly.
+		// The current code instead uses a copy of the data to populate an RHI texture. As this is DX11 only and is not the code path
+		// affecting main line hardware decoders (H.264/5), we do currently NOT optimize for this.
+
 		if (VideoDecoderOutputPC->GetOutputType() != FVideoDecoderOutputPC::EOutputType::Hardware_DX)
 		{
 			//
@@ -394,7 +415,7 @@ bool FElectraTextureSample::Convert(FTexture2DRHIRef& InDstTexture, const FConve
 		else
 		{
 			// We have a texture on the rendering device
-			Format = (VideoDecoderOutput->GetFormat() == PF_NV12) ? PF_G8 : PF_G16; // return a texture, representing the whole NV12/P010 layout as a single texture (1.5 height)
+			Format = VideoDecoderOutput->GetFormat();
 			bCrossDeviceCopy = false;
 		}
 	}
@@ -445,38 +466,9 @@ bool FElectraTextureSample::Convert(FTexture2DRHIRef& InDstTexture, const FConve
 	else
 #endif
 	{
-		// No. Then it better be a DX12 resource...
-		TRefCountPtr<ID3D12Resource> TextureDX12;
-		Res = TextureCommon->QueryInterface(__uuidof(ID3D12Resource), (void**)TextureDX12.GetInitReference());
-		check(Res == S_OK);
-		if (Res != S_OK)
-		{
-			return true;
-		}
-	
-		check(RHIGetInterfaceType() == ERHIInterfaceType::D3D12);
-
-		TRefCountPtr<ID3D12Fence> D3DFence;
-		if (SyncCommon)
-		{
-			Res = SyncCommon->QueryInterface(__uuidof(ID3D12Fence), (void**)D3DFence.GetInitReference());
-			check(Res == S_OK);
-			if (Res != S_OK)
-			{
-				return true;
-			}
-		}
-
-		// Copy data into RHI texture
-		if (RHICmdList.Bypass())
-		{
-			FRHICommandCopyResourceDX12 Cmd(TextureDX12, Texture, D3DCmdAllocator, D3DCmdList, D3DFence, SyncValue);
-			Cmd.Execute(RHICmdList);
-		}
-		else
-		{
-			new (RHICmdList.AllocCommand<FRHICommandCopyResourceDX12>()) FRHICommandCopyResourceDX12(TextureDX12, Texture, D3DCmdAllocator, D3DCmdList, D3DFence, SyncValue);
-		}
+		// We really only expect to be called for DX11
+		check("Unexpected call to conversion method from another RHI than the D3D11 one!");
+		return false;
 	}
 	return true;
 }
@@ -502,13 +494,17 @@ uint32 FElectraTextureSample::GetStride() const
 }
 
 
-IMFSample* FElectraTextureSample::GetMFSample()
+float FElectraTextureSample::GetSampleDataScale(bool b10Bit) const
 {
 	if (VideoDecoderOutput)
 	{
-		return VideoDecoderOutputPC->GetMFSample().GetReference();
+		auto Value = VideoDecoderOutputPC->GetDict().GetValue(IDecoderOutputOptionNames::PixelDataScale);
+		if (Value.IsValid() && Value.IsType(Electra::FVariantValue::EDataType::TypeDouble))
+		{
+			return (float)Value.GetDouble();
+		}
 	}
-	return nullptr;
+	return 1.0f;
 }
 
 #endif
diff --git a/Engine/Plugins/Media/ElectraUtil/Source/ElectraSamples/Public/Windows/ElectraTextureSample.h b/Engine/Plugins/Media/ElectraUtil/Source/ElectraSamples/Public/Windows/ElectraTextureSample.h
--- a/Engine/Plugins/Media/ElectraUtil/Source/ElectraSamples/Public/Windows/ElectraTextureSample.h
+++ b/Engine/Plugins/Media/ElectraUtil/Source/ElectraSamples/Public/Windows/ElectraTextureSample.h
@@ -44,17 +44,18 @@ public:
 	}
 
 #if WITH_ENGINE
-	virtual FRHITexture* GetTexture() const override
-	{
-		return Texture;
-	}
+	virtual FRHITexture* GetTexture() const override;
 #endif //WITH_ENGINE
 
-	IMFSample* GetMFSample();
-
 	virtual IMediaTextureSampleConverter* GetMediaTextureSampleConverter() override;
 
+#if !UE_SERVER
+	virtual void ShutdownPoolable() override;
+#endif
+
 private:
+	virtual float GetSampleDataScale(bool b10Bit) const override;
+
 	virtual bool Convert(FTexture2DRHIRef & InDstTexture, const FConversionHints & Hints) override;
 	virtual uint32 GetConverterInfoFlags() const
 	{
@@ -65,17 +66,13 @@ private:
 	EMediaTextureSampleFormat SampleFormat;
 
 	/** Destination Texture resource (from Rendering device) */
-	FTexture2DRHIRef Texture;
+	mutable FTexture2DRHIRef Texture;
 
 	/** True if texture format could support sRGB conversion in HW */
 	bool bCanUseSRGB;
 
 	/** Output data from video decoder. Baseclass holds reference */
 	FVideoDecoderOutputPC* VideoDecoderOutputPC;
-
-	/** DX12 related resources */
-	TRefCountPtr<ID3D12GraphicsCommandList> D3DCmdList;
-	TRefCountPtr<ID3D12CommandAllocator> D3DCmdAllocator;
 };
 
 using FElectraTextureSamplePtr = TSharedPtr<FElectraTextureSample, ESPMode::ThreadSafe>;
diff --git a/Engine/Plugins/Media/HAPDecoderElectra/Source/HAPDecoderElectra/HAPDecoderElectra.Build.cs b/Engine/Plugins/Media/HAPDecoderElectra/Source/HAPDecoderElectra/HAPDecoderElectra.Build.cs
--- a/Engine/Plugins/Media/HAPDecoderElectra/Source/HAPDecoderElectra/HAPDecoderElectra.Build.cs
+++ b/Engine/Plugins/Media/HAPDecoderElectra/Source/HAPDecoderElectra/HAPDecoderElectra.Build.cs
@@ -1,5 +1,7 @@
 // Copyright Epic Games, Inc. All Rights Reserved.
 
+using System.IO;
+
 namespace UnrealBuildTool.Rules
 {
 	public class HAPDecoderElectra : ModuleRules
@@ -21,6 +23,19 @@ namespace UnrealBuildTool.Rules
                     "HAPLib",
                     "SnappyLib",
                     });
+
+				if (Target.Platform == UnrealTargetPlatform.Win64)
+				{
+					PublicDependencyModuleNames.Add("DirectX");
+					AddEngineThirdPartyPrivateStaticDependencies(Target, "DX12");
+
+					if (Target.WindowsPlatform.Architecture != UnrealArch.Arm64)
+					{
+						PublicAdditionalLibraries.AddRange(new string[] {
+							Path.Combine(Target.WindowsPlatform.DirectXLibDir, "dxerr.lib"),
+						});
+					}
+				}
 			}
 		}
 	}
diff --git a/Engine/Plugins/Media/HAPDecoderElectra/Source/HAPDecoderElectra/Private/HAPDecoder/ElectraMediaHAPDecoder.cpp b/Engine/Plugins/Media/HAPDecoderElectra/Source/HAPDecoderElectra/Private/HAPDecoder/ElectraMediaHAPDecoder.cpp
--- a/Engine/Plugins/Media/HAPDecoderElectra/Source/HAPDecoderElectra/Private/HAPDecoder/ElectraMediaHAPDecoder.cpp
+++ b/Engine/Plugins/Media/HAPDecoderElectra/Source/HAPDecoderElectra/Private/HAPDecoder/ElectraMediaHAPDecoder.cpp
@@ -14,6 +14,7 @@
 #include "IElectraDecoderOutputVideo.h"
 #include "IElectraDecoderResourceDelegate.h"
 #include "ElectraDecodersUtils.h"
+#include COMPILED_PLATFORM_HEADER(ElectraDecoderGPUBufferHelpers.h)
 
 #include "hap.h"
 
@@ -127,27 +128,54 @@ public:
 	{
 		if (InBufferIndex == 0)
 		{
-			return ColorBuffer;
+			return Buffer[Buffer_Color].Buffer;
 		}
 		else if (InBufferIndex == 1)
 		{
-			return AlphaBuffer;
+			return Buffer[Buffer_Alpha].Buffer;
 		}
 		return nullptr;
 	}
 	void* GetBufferTextureByIndex(int32 InBufferIndex) const override
 	{
+#if ELECTRA_MEDIAGPUBUFFER_DX12
+		if (InBufferIndex == 0)
+		{
+			return Buffer[Buffer_Color].GPUBuffer.Resource.GetReference();
+		}
+		if (InBufferIndex == 1)
+		{
+			return Buffer[Buffer_Alpha].GPUBuffer.Resource.GetReference();
+		}
+#endif
 		return nullptr;
 	}
+	virtual bool GetBufferTextureSyncByIndex(int32 InBufferIndex, FElectraDecoderOutputSync& SyncObject) const override
+	{
+#if ELECTRA_MEDIAGPUBUFFER_DX12
+		if (InBufferIndex == 0)
+		{
+			SyncObject = { Buffer[Buffer_Color].GPUBuffer.Fence.GetReference(), Buffer[Buffer_Color].GPUBuffer.FenceValue };
+			return true;
+		}
+		if (InBufferIndex == 1)
+		{
+			SyncObject = { Buffer[Buffer_Alpha].GPUBuffer.Fence.GetReference(), Buffer[Buffer_Alpha].GPUBuffer.FenceValue };
+			return true;
+		}
+#endif
+		return false;
+	}
+
 	EElectraDecoderPlatformPixelFormat GetBufferFormatByIndex(int32 InBufferIndex) const override
 	{
 		if (InBufferIndex == 0)
 		{
-			return ColorBufferFormat;
+			return Buffer[Buffer_Color].BufferFormat;
 		}
 		else if (InBufferIndex == 1)
 		{
-			return AlphaBufferFormat;
+			return Buffer[Buffer_Alpha].BufferFormat;
 		}
 	return EElectraDecoderPlatformPixelFormat::INVALID;
 	}
@@ -155,11 +183,11 @@ public:
 	{
 		if (InBufferIndex == 0)
 		{
-			return ColorBufferEncoding;
+			return Buffer[Buffer_Color].BufferEncoding;
 		}
 		else if (InBufferIndex == 1)
 		{
-			return AlphaBufferEncoding;
+			return Buffer[Buffer_Alpha].BufferEncoding;
 		}
 		return EElectraDecoderPlatformPixelEncoding::Native;
 	}
@@ -167,11 +195,11 @@ public:
 	{
 		if (InBufferIndex == 0)
 		{
-			return ColorPitch;
+			return Buffer[Buffer_Color].Pitch;
 		}
 		else if (InBufferIndex == 1)
 		{
-			return AlphaPitch;
+			return Buffer[Buffer_Alpha].Pitch;
 		}
 		return 0;
 	}
@@ -183,8 +211,6 @@ public:
 	FElectraVideoDecoderOutputCropValues Crop;
 	int32 Width = 0;
 	int32 Height = 0;
-	int32 ColorPitch = 0;
-	int32 AlphaPitch = 0;
 	int32 NumBits = 0;
 	int32 AspectW = 1;
 	int32 AspectH = 1;
@@ -195,12 +221,23 @@ public:
 
 	uint32 Codec4CC = 0;
 	int32 NumBuffers = 0;
-	TSharedPtr<TArray<uint8>, ESPMode::ThreadSafe> ColorBuffer;
-	TSharedPtr<TArray<uint8>, ESPMode::ThreadSafe> AlphaBuffer;
-	EElectraDecoderPlatformPixelFormat ColorBufferFormat;
-	EElectraDecoderPlatformPixelEncoding ColorBufferEncoding;
-	EElectraDecoderPlatformPixelFormat AlphaBufferFormat;
-	EElectraDecoderPlatformPixelEncoding AlphaBufferEncoding;
+
+	enum {
+		Buffer_Color = 0,
+		Buffer_Alpha,
+
+		Buffer_Max
+	};
+
+	struct FBufferInfo {
+		TSharedPtr<TArray<uint8>, ESPMode::ThreadSafe> Buffer;
+		EElectraDecoderPlatformPixelFormat BufferFormat;
+		EElectraDecoderPlatformPixelEncoding BufferEncoding;
+		int32 Pitch = 0;
+#if ELECTRA_MEDIAGPUBUFFER_DX12
+		FElectraMediaDecoderOutputBufferPool_DX12::FOutputData GPUBuffer;
+#endif
+	} Buffer[Buffer_Max];
 };
 
 
@@ -266,13 +303,24 @@ private:
 	int32 AspectW = 0;
 	int32 AspectH = 0;
 	uint32 Codec4CC = 0;
-	uint32 AllocSizeColor = 0;
-	uint32 AllocSizeAlpha = 0;
 
 	IElectraDecoder::FError LastError;
 
 	TSharedPtr<FVideoDecoderOutputHAPElectra, ESPMode::ThreadSafe> CurrentOutput;
 	bool bFlushPending = false;
+
+	TWeakPtr<IElectraDecoderResourceDelegate, ESPMode::ThreadSafe> ResourceDelegate;
+
+	uint32 MaxOutputBuffers;
+	struct FBufferInstanceInfo
+	{
+		uint32 AllocSize = 0;
+		uint32 AllocBlockSize = 0;
+#if ELECTRA_MEDIAGPUBUFFER_DX12
+		mutable TSharedPtr<FElectraMediaDecoderOutputBufferPool_DX12> D3D12ResourcePool;
+#endif
+	} InstanceBufferData[FVideoDecoderOutputHAPElectra::Buffer_Max];
+	mutable TUniquePtr<IElectraDecoderResourceDelegateBase::IAsyncConsecutiveTaskSync> TaskSync;
 };
 
 
@@ -341,6 +389,8 @@ void FElectraMediaHAPDecoder::Shutdown()
 
 FVideoDecoderHAPElectra::FVideoDecoderHAPElectra(const TMap<FString, FVariant>& InOptions, TSharedPtr<IElectraDecoderResourceDelegate, ESPMode::ThreadSafe> InResourceDelegate)
 {
+	ResourceDelegate = InResourceDelegate;
+
 	DisplayWidth = (int32)ElectraDecodersUtil::GetVariantValueSafeI64(InOptions, TEXT("width"), 0);
 	DisplayHeight = (int32)ElectraDecodersUtil::GetVariantValueSafeI64(InOptions, TEXT("height"), 0);
 	// All HAP formats have 4x4 pixel blocks.
@@ -355,24 +405,33 @@ FVideoDecoderHAPElectra::FVideoDecoderHAPElectra(const TMap<FString, FVariant>&
 		case Make4CC('H','a','p','1'):
 		case Make4CC('H','a','p','A'):
 		{
-			AllocSizeColor = DecodedWidth * DecodedHeight / 2;
-			AllocSizeAlpha = DecodedWidth * DecodedHeight / 2;
+			InstanceBufferData[FVideoDecoderOutputHAPElectra::Buffer_Color].AllocSize = DecodedWidth * DecodedHeight / 2;
+			InstanceBufferData[FVideoDecoderOutputHAPElectra::Buffer_Alpha].AllocSize = DecodedWidth * DecodedHeight / 2;
+			InstanceBufferData[FVideoDecoderOutputHAPElectra::Buffer_Color].AllocBlockSize = 8;
+			InstanceBufferData[FVideoDecoderOutputHAPElectra::Buffer_Alpha].AllocBlockSize = 8;
 			break;
 		}
 		case Make4CC('H','a','p','Y'):
 		case Make4CC('H','a','p','5'):
 		{
-			AllocSizeColor = DecodedWidth * DecodedHeight;
-			AllocSizeAlpha = DecodedWidth * DecodedHeight;
+			InstanceBufferData[FVideoDecoderOutputHAPElectra::Buffer_Color].AllocSize = DecodedWidth * DecodedHeight;
+			InstanceBufferData[FVideoDecoderOutputHAPElectra::Buffer_Alpha].AllocSize = DecodedWidth * DecodedHeight;
+			InstanceBufferData[FVideoDecoderOutputHAPElectra::Buffer_Color].AllocBlockSize = 16;
+			InstanceBufferData[FVideoDecoderOutputHAPElectra::Buffer_Alpha].AllocBlockSize = 16;
 			break;
 		}
 		default:
 		{
-			AllocSizeColor = DecodedWidth * DecodedHeight * 4;
-			AllocSizeAlpha = DecodedWidth * DecodedHeight * 4;
+			InstanceBufferData[FVideoDecoderOutputHAPElectra::Buffer_Color].AllocSize = DecodedWidth * DecodedHeight * 4;
+			InstanceBufferData[FVideoDecoderOutputHAPElectra::Buffer_Alpha].AllocSize = DecodedWidth * DecodedHeight * 4;
+			InstanceBufferData[FVideoDecoderOutputHAPElectra::Buffer_Color].AllocBlockSize = 0;
+			InstanceBufferData[FVideoDecoderOutputHAPElectra::Buffer_Alpha].AllocBlockSize = 0;
 			break;
 		}
 	}
+
+	MaxOutputBuffers = (uint32)ElectraDecodersUtil::GetVariantValueSafeU64(InOptions, TEXT("max_output_buffers"), 5);
+	MaxOutputBuffers += kElectraDecoderPipelineExtraFrames;
 }
 
 FVideoDecoderHAPElectra::~FVideoDecoderHAPElectra()
@@ -499,6 +558,17 @@ IElectraDecoder::EDecoderError FVideoDecoderHAPElectra::DecodeAccessUnit(const F
 		return IElectraDecoder::EDecoderError::NoBuffer;
 	}
 
+#if ELECTRA_MEDIAGPUBUFFER_DX12
+	// If we will create a new resource pool or we have still buffers in an existing one, we can proceed, else we'd have no resources to output the data
+	for (uint32 BufIdx = 0; BufIdx < FVideoDecoderOutputHAPElectra::Buffer_Max; ++BufIdx)
+	{
+		if (InstanceBufferData[BufIdx].D3D12ResourcePool.IsValid() && !InstanceBufferData[BufIdx].D3D12ResourcePool->BufferAvailable())
+		{
+			return IElectraDecoder::EDecoderError::NoBuffer;
+		}
+	}
+#endif
+
 	// Decode data. This immediately produces a new output frame.
 	if (InInputAccessUnit.Data && InInputAccessUnit.DataSize)
 	{
@@ -510,6 +580,18 @@ IElectraDecoder::EDecoderError FVideoDecoderHAPElectra::DecodeAccessUnit(const F
 			return IElectraDecoder::EDecoderError::Error;
 		}
 
+		void* PlatformDevice = nullptr;
+		int32 PlatformDeviceVersion = 0;
+		bool bUseGPUBuffers = false;
+#if ELECTRA_MEDIAGPUBUFFER_DX12
+		auto PinnedResourceDelegate = ResourceDelegate.Pin();
+		if (PinnedResourceDelegate.IsValid())
+		{
+			PinnedResourceDelegate->GetD3DDevice(&PlatformDevice, &PlatformDeviceVersion);
+			bUseGPUBuffers = (PlatformDevice && PlatformDeviceVersion >= 12000);
+		}
+#endif
+
 		TSharedPtr<FVideoDecoderOutputHAPElectra, ESPMode::ThreadSafe> NewOutput = MakeShared<FVideoDecoderOutputHAPElectra>();
 		NewOutput->PTS = InInputAccessUnit.PTS;
 		NewOutput->UserValue = InInputAccessUnit.UserValue;
@@ -523,54 +605,135 @@ IElectraDecoder::EDecoderError FVideoDecoderHAPElectra::DecodeAccessUnit(const F
 			NewOutput->AspectW = AspectW;
 			NewOutput->AspectH = AspectH;
 		}
-	// These actually depend on the format, but we don't support the latest HAP formats at the moment.
+// These actually depend on the format, but we don't support the latest HAP formats at the moment.
 		NewOutput->NumBits = 8;
 		NewOutput->PixelFormat = 0;
 
 		NewOutput->ExtraValues.Emplace(TEXT("codec"), FVariant(TEXT("hap")));
 		NewOutput->ExtraValues.Emplace(TEXT("codec_4cc"), FVariant(Codec4CC));
 
-		NewOutput->ColorBuffer = MakeShared<TArray<uint8>, ESPMode::ThreadSafe>();
-		NewOutput->ColorBuffer->AddUninitialized(AllocSizeColor);
+		// Setup all buffers / textures (we clamp to what out max number is)
+		uint32 NumBuffers = FMath::Min(TextureCount, (uint32)FVideoDecoderOutputHAPElectra::Buffer_Max);
+		for (uint32 BufIdx = 0; BufIdx < NumBuffers; ++BufIdx)
+		{
+			auto& BufferInfo = NewOutput->Buffer[BufIdx];
+
+			uint8* BufferAddr;
+			uint32 BufferPitch = 0;
+			uint32 ResourcePitch = 0;
+			uint32 BlockSizeX = InstanceBufferData[BufIdx].AllocBlockSize ? 4 : 1;
+			uint32 BlockSizeY = InstanceBufferData[BufIdx].AllocBlockSize ? 4 : 1;
+			uint32 BlockBytes = InstanceBufferData[BufIdx].AllocBlockSize ? InstanceBufferData[BufIdx].AllocBlockSize : 4;
+
+#if ELECTRA_MEDIAGPUBUFFER_DX12
+			TSharedPtr<TArray<uint8>, ESPMode::ThreadSafe> TempBuffer;
+			if (bUseGPUBuffers)
+			{
+				TRefCountPtr D3D12Device(static_cast<ID3D12Device*>(PlatformDevice));
+
+				// Create the resource pool as needed...
+				if (!InstanceBufferData[BufIdx].D3D12ResourcePool.IsValid())
+				{
+					// We always expect to get BCx style data here, but allow for a default RGBA8 setup (as does the initialization code)
+					InstanceBufferData[BufIdx].D3D12ResourcePool = MakeShared<FElectraMediaDecoderOutputBufferPool_DX12, ESPMode::ThreadSafe>(D3D12Device, MaxOutputBuffers, DecodedWidth / BlockSizeX, DecodedHeight / BlockSizeY, BlockBytes);
+				}
+
+				// Create a tasksync instance so we can have our own async jobs run in consecutive order
+				if (!TaskSync.IsValid())
+				{
+					TaskSync.Reset(PinnedResourceDelegate->CreateAsyncConsecutiveTaskSync());
+				}
+
+				InstanceBufferData[BufIdx].D3D12ResourcePool->AllocateOutputDataAsBuffer(BufferInfo.GPUBuffer, ResourcePitch);
+
+				// Make the decoder write to a cached, temp buffer (otherwise performance is quite bad)
+				TempBuffer = MakeShared<TArray<uint8>, ESPMode::ThreadSafe>();
+				if (!TempBuffer.IsValid())
+				{
+					PostError(Result, TEXT("Could not allocate HAP decoder output buffer"), ERRCODE_INTERNAL_FAILED_TO_DECODE_INPUT);
+					return IElectraDecoder::EDecoderError::Error;
+				}
+				TempBuffer->AddUninitialized(InstanceBufferData[BufIdx].AllocSize);
+				BufferAddr = TempBuffer->GetData();
+				if (BufferAddr == nullptr)
+				{
+					PostError(Result, TEXT("Could not allocate HAP decoder output buffer"), ERRCODE_INTERNAL_FAILED_TO_DECODE_INPUT);
+					return IElectraDecoder::EDecoderError::Error;
+				}
+				BufferPitch = (DecodedWidth / BlockSizeX) * BlockBytes;
+			}
+			else
+#endif
+			{
+				BufferInfo.Buffer = MakeShared<TArray<uint8>, ESPMode::ThreadSafe>();
+				BufferInfo.Buffer->AddUninitialized(InstanceBufferData[BufIdx].AllocSize);
+				BufferAddr = BufferInfo.Buffer->GetData();
+				if (BufferAddr == nullptr)
+				{
+					PostError(Result, TEXT("Could not allocate HAP decoder output buffer"), ERRCODE_INTERNAL_FAILED_TO_DECODE_INPUT);
+					return IElectraDecoder::EDecoderError::Error;
+				}
+			}
 
-		unsigned long ColorBufferBytesUsed = 0;
-		unsigned int ColorBufferTextureFormat = 0;
+			unsigned long BufferBytesUsed = 0;
+			unsigned int BufferTextureFormat = 0;
 
-		Result = static_cast<HapResult>(HapDecode(InInputAccessUnit.Data, static_cast<unsigned long>(InInputAccessUnit.DataSize), 0, HapDecodeCallback, nullptr,
-			NewOutput->ColorBuffer->GetData(), AllocSizeColor, &ColorBufferBytesUsed, &ColorBufferTextureFormat));
-		if (Result != HapResult_No_Error)
-		{
-			PostError(Result, TEXT("HapDecode() failed"), ERRCODE_INTERNAL_FAILED_TO_DECODE_INPUT);
-			return IElectraDecoder::EDecoderError::Error;
-		}
-		NewOutput->ColorBuffer->SetNumUnsafeInternal((int32) ColorBufferBytesUsed);
-		ConvertHapTextureFormat(static_cast<HapTextureFormat>(ColorBufferTextureFormat), NewOutput->ColorBufferFormat, NewOutput->ColorBufferEncoding);
+			Result = static_cast<HapResult>(HapDecode(InInputAccessUnit.Data, static_cast<unsigned long>(InInputAccessUnit.DataSize), 0, HapDecodeCallback, nullptr,
+													  BufferAddr, InstanceBufferData[BufIdx].AllocSize, &BufferBytesUsed, &BufferTextureFormat));
 
-		NewOutput->ColorPitch = GetImageBufferPitch(static_cast<HapTextureFormat>(ColorBufferTextureFormat), NewOutput->Width);
+#if ELECTRA_MEDIAGPUBUFFER_DX12
+			if (bUseGPUBuffers)
+			{
+				if (Result == HapResult_No_Error)
+				{
+					if (!PinnedResourceDelegate->RunCodeAsync([BufferResource = BufferInfo.GPUBuffer.Resource, BufferFence = BufferInfo.GPUBuffer.Fence, BufferFenceValue = BufferInfo.GPUBuffer.FenceValue, ResourcePitch, BufferAddr, BufferPitch, DecodedPitch = DecodedHeight / BlockSizeY, TempBuffer]()
+						{
+							uint8* DX12BufferAddr;
+							HRESULT Res = BufferResource->Map(0, nullptr, (void**)&DX12BufferAddr);
+							check(SUCCEEDED(Res));
 
-		NewOutput->NumBuffers = (int32)TextureCount;
-		NewOutput->Codec4CC = Codec4CC;
+							FElectraMediaDecoderOutputBufferPool_DX12::CopyWithPitchAdjust(DX12BufferAddr, ResourcePitch, BufferAddr, BufferPitch, DecodedPitch);
 
-		// Alpha?
-		if (TextureCount == 2)
-		{
-			NewOutput->AlphaBuffer = MakeShared<TArray<uint8>, ESPMode::ThreadSafe>();
-			NewOutput->AlphaBuffer->AddUninitialized(AllocSizeAlpha);
+							BufferResource->Unmap(0, nullptr);
+
+							BufferFence->Signal(BufferFenceValue);
+						}, TaskSync.Get()))
+					{
+						uint8* DX12BufferAddr;
+						HRESULT Res = BufferInfo.GPUBuffer.Resource->Map(0, nullptr, (void**)&DX12BufferAddr);
+						check(SUCCEEDED(Res));
+
+						FElectraMediaDecoderOutputBufferPool_DX12::CopyWithPitchAdjust(DX12BufferAddr, ResourcePitch, BufferAddr, BufferPitch, DecodedHeight / BlockSizeY);
+
+						BufferInfo.GPUBuffer.Resource->Unmap(0, nullptr);
 
-			unsigned long AlphaBufferBytesUsed = 0;
-			unsigned int AlphaBufferTextureFormat = 0;
+						BufferInfo.GPUBuffer.Fence->Signal(BufferInfo.GPUBuffer.FenceValue);
+					}
+
+					BufferInfo.Pitch = ResourcePitch;
+				}
+				TempBuffer.Reset();
+			}
+			else
+#endif
+			{
+				if (Result == HapResult_No_Error)
+				{
+					BufferInfo.Buffer->SetNumUnsafeInternal((int32)BufferBytesUsed);
+					BufferInfo.Pitch = GetImageBufferPitch(static_cast<HapTextureFormat>(BufferTextureFormat), NewOutput->Width);
+				}
+			}
 
-			Result = static_cast<HapResult>(HapDecode(InInputAccessUnit.Data, static_cast<unsigned long>(InInputAccessUnit.DataSize), 1, HapDecodeCallback, nullptr,
-				NewOutput->ColorBuffer->GetData(), AllocSizeAlpha, &AlphaBufferBytesUsed, &AlphaBufferTextureFormat));
 			if (Result != HapResult_No_Error)
 			{
 				PostError(Result, TEXT("HapDecode() failed"), ERRCODE_INTERNAL_FAILED_TO_DECODE_INPUT);
 				return IElectraDecoder::EDecoderError::Error;
 			}
-			NewOutput->AlphaBuffer->SetNumUnsafeInternal((int32) AlphaBufferBytesUsed);
-			ConvertHapTextureFormat(static_cast<HapTextureFormat>(AlphaBufferTextureFormat), NewOutput->AlphaBufferFormat, NewOutput->AlphaBufferEncoding);
 
-			NewOutput->AlphaPitch = GetImageBufferPitch(static_cast<HapTextureFormat>(AlphaBufferTextureFormat), NewOutput->Width);
+			ConvertHapTextureFormat(static_cast<HapTextureFormat>(BufferTextureFormat), BufferInfo.BufferFormat, BufferInfo.BufferEncoding);
+
+			NewOutput->NumBuffers = (int32)TextureCount;
+			NewOutput->Codec4CC = Codec4CC;
 		}
 
 		CurrentOutput = MoveTemp(NewOutput);
diff --git a/Engine/Plugins/Media/VPxDecoderElectra/Source/VPxDecoderElectra/Private/VPxDecoder/ElectraMediaVPxDecoder.cpp b/Engine/Plugins/Media/VPxDecoderElectra/Source/VPxDecoderElectra/Private/VPxDecoder/ElectraMediaVPxDecoder.cpp
--- a/Engine/Plugins/Media/VPxDecoderElectra/Source/VPxDecoderElectra/Private/VPxDecoder/ElectraMediaVPxDecoder.cpp
+++ b/Engine/Plugins/Media/VPxDecoderElectra/Source/VPxDecoderElectra/Private/VPxDecoder/ElectraMediaVPxDecoder.cpp
@@ -16,9 +16,22 @@
 #include "ElectraDecodersUtils.h"
 #include "Utils/Google/ElectraUtilsVPxVideo.h"
 #include "Containers/Queue.h"
+#include <Stats/Stats.h>
+
+#include COMPILED_PLATFORM_HEADER(ElectraDecoderGPUBufferHelpers.h)
+
 #include <vpx/vpx_decoder.h>
 #include <vpx/vp8dx.h>
 
+/*********************************************************************************************************************/
+
+
+/*********************************************************************************************************************/
+
+DECLARE_CYCLE_STAT(TEXT("ElectraDecoder ConvertOutput"), STAT_ElectraDecoder_ConvertOutputVpx, STATGROUP_Media);
+
+/*********************************************************************************************************************/
+
 #define ERRCODE_INTERNAL_NO_ERROR							0
 #define ERRCODE_INTERNAL_ALREADY_CLOSED						1
 #define ERRCODE_INTERNAL_COULD_NOT_CREATE_DECODER			2
@@ -41,10 +54,6 @@ public:
 
 };
 
-// This should really be handled differently...
-static bool ConvertDecodedImageToNV12(TSharedPtr<TArray<uint8>, ESPMode::ThreadSafe>& OutNV12Buffer, const vpx_image_t* InDecodedImage);
-static bool ConvertDecodedImageToP010(TSharedPtr<TArray<uint8>, ESPMode::ThreadSafe>& OutP010Buffer, const vpx_image_t* InDecodedImage);
-
 
 class FVideoDecoderOutputVPxElectra : public IElectraDecoderVideoOutput, public IElectraDecoderVideoOutputImageBuffers
 {
@@ -108,7 +117,26 @@ public:
 		return nullptr;
 	}
 	void* GetBufferTextureByIndex(int32 InBufferIndex) const override
-	{ return nullptr; }
+	{
+#if ELECTRA_MEDIAGPUBUFFER_DX12
+		if (InBufferIndex == 0)
+		{
+			return GPUBuffer.Resource.GetReference();
+		}
+#endif
+		return nullptr;
+	}
+	virtual bool GetBufferTextureSyncByIndex(int32 InBufferIndex, FElectraDecoderOutputSync& SyncObject) const override
+	{
+#if ELECTRA_MEDIAGPUBUFFER_DX12
+		if (InBufferIndex == 0)
+		{
+			SyncObject = { GPUBuffer.Fence.GetReference(), GPUBuffer.FenceValue };
+			return true;
+		}
+#endif
+		return false;
+	}
 	EElectraDecoderPlatformPixelFormat GetBufferFormatByIndex(int32 InBufferIndex) const override
 	{
 		if (InBufferIndex == 0)
@@ -157,6 +185,9 @@ public:
 	EElectraDecoderPlatformPixelFormat ColorBufferFormat;
 	EElectraDecoderPlatformPixelEncoding ColorBufferEncoding;
 	int32 ColorPitch = 0;
+#if ELECTRA_MEDIAGPUBUFFER_DX12
+	FElectraMediaDecoderOutputBufferPool_DX12::FOutputData GPUBuffer;
+#endif
 };
 
 
@@ -239,6 +270,9 @@ private:
 	};
 	EConvertResult ConvertDecoderOutput(vpx_image_t* InDecodedImage);
 
+	bool ConvertDecodedImageToNV12orP010(FVideoDecoderOutputVPxElectra* InNewOutput, const vpx_image_t* InDecodedImage, void* PlatformDevice, int32 PlatformDeviceVersion) const;
+
+	TWeakPtr<IElectraDecoderResourceDelegate, ESPMode::ThreadSafe> ResourceDelegate;
 
 	TQueue<TSharedPtr<FDecoderInput, ESPMode::NotThreadSafe>, EQueueMode::Spsc> PendingDecoderInput;
 	TArray<TSharedPtr<FDecoderInput, ESPMode::NotThreadSafe>> InDecoderInput;
@@ -250,8 +284,14 @@ private:
 	vpx_codec_ctx_t DecoderContext;
 	vpx_codec_ctx_t* VideoDecoderHandle = nullptr;
 	vpx_codec_iter_t VideoDecoderOutputIterator = nullptr;
-};
 
+	uint32 MaxWidth;
+	uint32 MaxHeight;
+	uint32 MaxOutputBuffers;
+#if ELECTRA_MEDIAGPUBUFFER_DX12
+	mutable TSharedPtr<FElectraMediaDecoderOutputBufferPool_DX12> D3D12ResourcePool;
+#endif
+};
 
 /*********************************************************************************************************************/
 /*********************************************************************************************************************/
@@ -340,6 +380,13 @@ void FElectraMediaVPxDecoder::Shutdown()
 FVideoDecoderVPxElectra::FVideoDecoderVPxElectra(const TMap<FString, FVariant>& InOptions, TSharedPtr<IElectraDecoderResourceDelegate, ESPMode::ThreadSafe> InResourceDelegate)
 {
 	Codec4CC = (uint32)ElectraDecodersUtil::GetVariantValueSafeU64(InOptions, TEXT("codec_4cc"), 0);
+	ResourceDelegate = InResourceDelegate;
+
+	MaxWidth = (uint32)ElectraDecodersUtil::GetVariantValueSafeU64(InOptions, TEXT("max_width"), 1920);
+	MaxHeight = (uint32)ElectraDecodersUtil::GetVariantValueSafeU64(InOptions, TEXT("max_height"), 1080);
+
+	MaxOutputBuffers = (uint32)ElectraDecodersUtil::GetVariantValueSafeU64(InOptions, TEXT("max_output_buffers"), 5);
+	MaxOutputBuffers += kElectraDecoderPipelineExtraFrames;
 }
 
 FVideoDecoderVPxElectra::~FVideoDecoderVPxElectra()
@@ -428,6 +475,14 @@ IElectraDecoder::EDecoderError FVideoDecoderVPxElectra::DecodeAccessUnit(const F
 		return IElectraDecoder::EDecoderError::NoBuffer;
 	}
 
+#if ELECTRA_MEDIAGPUBUFFER_DX12
+	// If we will create a new resource pool or we have still buffers in an existing one, we can proceed, else we'd have no resources to output the data
+	if (D3D12ResourcePool.IsValid() && !D3D12ResourcePool->BufferAvailable())
+	{
+		return IElectraDecoder::EDecoderError::NoBuffer;
+	}
+#endif
+
 	// CSD only buffer is not handled at the moment.
 	check((InInputAccessUnit.Flags & EElectraDecoderFlags::InitCSDOnly) == EElectraDecoderFlags::None);
 
@@ -757,6 +812,8 @@ FVideoDecoderVPxElectra::EConvertResult FVideoDecoderVPxElectra::ConvertDecoderO
 		return EConvertResult::GotEOS;
 	}
 
+	SCOPE_CYCLE_COUNTER(STAT_ElectraDecoder_ConvertOutputVpx);
+
 	// Find the input corresponding to this output.
 	TSharedPtr<FDecoderInput, ESPMode::NotThreadSafe> In;
 	for(int32 nInDec=0; nInDec<InDecoderInput.Num(); ++nInDec)
@@ -784,9 +841,20 @@ FVideoDecoderVPxElectra::EConvertResult FVideoDecoderVPxElectra::ConvertDecoderO
 	NewOutput->AspectW = 1;
 	NewOutput->AspectH = 1;
 
+	void* PlatformDevice = nullptr;
+	int32 PlatformDeviceVersion = 0;
+	bool bUseGPUBuffers = false;
+#if ELECTRA_MEDIAGPUBUFFER_DX12
+	if (auto PinnedResourceDelegate = ResourceDelegate.Pin())
+	{
+		PinnedResourceDelegate->GetD3DDevice(&PlatformDevice, &PlatformDeviceVersion);
+		bUseGPUBuffers = (PlatformDevice && PlatformDeviceVersion >= 12000);
+	}
+#endif
+
 	if (InDecodedImage->fmt == VPX_IMG_FMT_I420)
 	{
-		if (!ConvertDecodedImageToNV12(NewOutput->ColorBuffer, InDecodedImage))
+		if (!ConvertDecodedImageToNV12orP010(NewOutput.Get(), InDecodedImage, PlatformDevice, PlatformDeviceVersion))
 		{
 			PostError(0, FString::Printf(TEXT("Failed to convert decoded image")), ERRCODE_INTERNAL_FAILED_TO_CONVERT_OUTPUT_SAMPLE);
 			return EConvertResult::Failure;
@@ -794,14 +862,13 @@ FVideoDecoderVPxElectra::EConvertResult FVideoDecoderVPxElectra::ConvertDecoderO
 		NewOutput->NumBuffers = 1;
 		NewOutput->ColorBufferFormat = EElectraDecoderPlatformPixelFormat::NV12;
 		NewOutput->ColorBufferEncoding = EElectraDecoderPlatformPixelEncoding::Native;
-		NewOutput->ColorPitch = NewOutput->Width;
 		NewOutput->DecodedWidth = NewOutput->Width;
-		NewOutput->DecodedHeight = NewOutput->Height * 3 / 2;
+		NewOutput->DecodedHeight = bUseGPUBuffers ? NewOutput->Height : (NewOutput->Height * 3 / 2);
 	}
 	else if (InDecodedImage->fmt == VPX_IMG_FMT_I42016)
 	{
 		check(NewOutput->NumBits == 10);
-		if (!ConvertDecodedImageToP010(NewOutput->ColorBuffer, InDecodedImage))
+		if (!ConvertDecodedImageToNV12orP010(NewOutput.Get(), InDecodedImage, PlatformDevice, PlatformDeviceVersion))
 		{
 			PostError(0, FString::Printf(TEXT("Failed to convert decoded image")), ERRCODE_INTERNAL_FAILED_TO_CONVERT_OUTPUT_SAMPLE);
 			return EConvertResult::Failure;
@@ -809,9 +876,11 @@ FVideoDecoderVPxElectra::EConvertResult FVideoDecoderVPxElectra::ConvertDecoderO
 		NewOutput->NumBuffers = 1;
 		NewOutput->ColorBufferFormat = EElectraDecoderPlatformPixelFormat::P010;
 		NewOutput->ColorBufferEncoding = EElectraDecoderPlatformPixelEncoding::Native;
-		NewOutput->ColorPitch = NewOutput->Width * 2;
 		NewOutput->DecodedWidth = NewOutput->Width;
-		NewOutput->DecodedHeight = NewOutput->Height * 3 / 2;
+		NewOutput->DecodedHeight = bUseGPUBuffers ? NewOutput->Height : (NewOutput->Height * 3 / 2);
+
+		// VPx decoders return the 10-bit output in the lower bits, but the output pipe expects it in the upper bits. Post scale to compensate!
+		NewOutput->ExtraValues.Emplace(TEXT("pix_datascale"), FVariant(64.0));
 	}
 	else
 	{
@@ -838,24 +907,18 @@ FVideoDecoderVPxElectra::EConvertResult FVideoDecoderVPxElectra::ConvertDecoderO
 	return EConvertResult::Success;
 }
 
-
-static bool ConvertDecodedImageToNV12(TSharedPtr<TArray<uint8>, ESPMode::ThreadSafe>& OutNV12Buffer, const vpx_image_t* InDecodedImage)
+bool FVideoDecoderVPxElectra::ConvertDecodedImageToNV12orP010(FVideoDecoderOutputVPxElectra* InNewOutput, const vpx_image_t* InDecodedImage, void* PlatformDevice, int32 PlatformDeviceVersion) const
 {
+	const bool bIsNV12 = (InDecodedImage->fmt == VPX_IMG_FMT_I420);
+
 	const int32 w = InDecodedImage->d_w;
 	const int32 h = InDecodedImage->d_h;
 	const int32 aw = Align(w, 2);
 	const int32 ah = Align(h, 2);
 
-	int32 AllocSize = aw * ah * 3/2;
-
-	OutNV12Buffer = MakeShared<TArray<uint8>, ESPMode::ThreadSafe>();
-	OutNV12Buffer->AddUninitialized(AllocSize);
-
-	uint8* DstY = OutNV12Buffer->GetData();
-	uint8* DstUV = DstY + aw * ah;
-	const uint8* SrcY = (const uint8*) InDecodedImage->planes[0];
-	const uint8* SrcU = (const uint8*) InDecodedImage->planes[1];
-	const uint8* SrcV = (const uint8*) InDecodedImage->planes[2];
+	const uint8* SrcY = (const uint8*)InDecodedImage->planes[0];
+	const uint8* SrcU = (const uint8*)InDecodedImage->planes[1];
+	const uint8* SrcV = (const uint8*)InDecodedImage->planes[2];
 	const int32 PitchY = InDecodedImage->stride[0];
 	const int32 PitchU = InDecodedImage->stride[1];
 	const int32 PitchV = InDecodedImage->stride[2];
@@ -863,71 +926,117 @@ static bool ConvertDecodedImageToNV12(TSharedPtr<TArray<uint8>, ESPMode::ThreadS
 	{
 		return false;
 	}
-	for(int32 y=0; y<h; ++y)
+
+	uint8* DstY;
+	uint8* DstUV;
+	uint32 DstPitch;
+
+	TSharedPtr<TArray<uint8>, ESPMode::ThreadSafe>& OutNV12Buffer = InNewOutput->ColorBuffer;
+#if ELECTRA_MEDIAGPUBUFFER_DX12
+	if (!PlatformDevice || PlatformDeviceVersion < 12000)
+#endif
 	{
-		FMemory::Memcpy(DstY, SrcY, w);
-		DstY += aw;
-		SrcY += PitchY;
+		uint32 Pitch = bIsNV12 ? aw : (aw * 2);
+		int32 AllocSize = Pitch * (ah * 3 / 2);
+
+#if ELECTRA_MEDIAGPUBUFFER_DX12
+		InNewOutput->GPUBuffer.Resource = nullptr;
+		InNewOutput->GPUBuffer.Fence = nullptr;
+#endif
+
+		OutNV12Buffer = MakeShared<TArray<uint8>, ESPMode::ThreadSafe>();
+		OutNV12Buffer->AddUninitialized(AllocSize);
+		DstY = OutNV12Buffer->GetData();
+		DstUV = DstY + Pitch * ah;
+		DstPitch = Pitch;
 	}
-	int32 padUV = (aw - w) * 2;
-	for(int32 v=0; v<h/2; ++v)
+#if ELECTRA_MEDIAGPUBUFFER_DX12
+	else
 	{
-		for(int32 u=0; u<w/2; ++u)
+		TRefCountPtr D3D12Device(static_cast<ID3D12Device*>(PlatformDevice));
+
+		OutNV12Buffer.Reset();
+		InNewOutput->GPUBuffer.Resource = nullptr;
+
+		// Create the resource pool as needed...
+		if (!D3D12ResourcePool)
 		{
-			*DstUV++ = SrcU[u];
-			*DstUV++ = SrcV[u];
+			D3D12ResourcePool = MakeShared<FElectraMediaDecoderOutputBufferPool_DX12, ESPMode::ThreadSafe>(D3D12Device, MaxOutputBuffers, MaxWidth, MaxHeight * 3 / 2, bIsNV12 ? 1 : 2);
 		}
-		SrcU += PitchU;
-		SrcV += PitchV;
-		DstUV += padUV;
-	}
-	return true;
-}
-
-static bool ConvertDecodedImageToP010(TSharedPtr<TArray<uint8>, ESPMode::ThreadSafe>& OutP010Buffer, const vpx_image_t* InDecodedImage)
-{
-	const int32 w = InDecodedImage->d_w;
-	const int32 h = InDecodedImage->d_h;
-	const int32 aw = Align(w, 2);
-	const int32 ah = Align(h, 2);
 
-	int32 AllocSize = aw * ah * 3;
+		// Request resource and fence...
+		uint32 BufferPitch;
+		D3D12ResourcePool->AllocateOutputDataAsBuffer(InNewOutput->GPUBuffer, BufferPitch);
 
-	OutP010Buffer = MakeShared<TArray<uint8>, ESPMode::ThreadSafe>();
-	OutP010Buffer->AddUninitialized(AllocSize);
+		// Correct pitch ton reflect resource's setup
+		InNewOutput->ColorPitch = BufferPitch;
 
-	uint16* DstY = (uint16*)OutP010Buffer->GetData();
-	uint16* DstUV = DstY + aw * ah;
-	const uint16* SrcY = (const uint16*) InDecodedImage->planes[0];
-	const uint16* SrcU = (const uint16*) InDecodedImage->planes[1];
-	const uint16* SrcV = (const uint16*) InDecodedImage->planes[2];
-	const int32 PitchY = InDecodedImage->stride[0] / 2;
-	const int32 PitchU = InDecodedImage->stride[1] / 2;
-	const int32 PitchV = InDecodedImage->stride[2] / 2;
-	if (!SrcY || !SrcU || !SrcV)
-	{
-		return false;
+		// Map the buffer so we can get a CPU address to the WC configured buffer
+		HRESULT Res = InNewOutput->GPUBuffer.Resource->Map(0, nullptr, (void**)&DstY);
+		check(SUCCEEDED(Res));
+		DstUV = DstY + BufferPitch * ah;
+		DstPitch = BufferPitch;
 	}
-	for(int32 y=0; y<h; ++y)
+#endif // ELECTRA_MEDIAGPUBUFFER_DX12
+
+	InNewOutput->ColorPitch = DstPitch;
+
+	if (bIsNV12)
 	{
-		for(int32 x=0; x<w; ++x)
+		for (int32 y = 0; y < h; ++y)
 		{
-			DstY[x] = SrcY[x] << 6;
+			FMemory::Memcpy(DstY, SrcY, w);
+			SrcY += PitchY;
+			DstY += DstPitch;
+		}
+		for (int32 v = 0; v < h / 2; ++v)
+		{
+			uint8* DstUVLine = DstUV;
+			for (int32 u = 0; u < w / 2; ++u)
+			{
+				*DstUVLine++ = SrcU[u];
+				*DstUVLine++ = SrcV[u];
+			}
+			SrcU += PitchU;
+			SrcV += PitchV;
+			DstUV += DstPitch;
 		}
-		DstY += aw;
-		SrcY += PitchY;
 	}
-	int32 padUV = (aw - w) * 2;
-	for(int32 v=0; v<h/2; ++v)
+	else
 	{
-		for(int32 u=0; u<w/2; ++u)
+		// note: data is delivered in the lower 10-bits, but expected in the upper
+		// -> instead of processing the data here, we provide a "data scale" attribute to be applied on conversion from YUV to RGB
+		for (int32 y = 0; y < h; ++y)
+		{
+			FMemory::Memcpy(DstY, SrcY, int64(w) << 1);
+			SrcY += PitchY;
+			DstY += DstPitch;
+		}
+		for (int32 v = 0; v < h / 2; ++v)
 		{
-			*DstUV++ = SrcU[u] << 6;
-			*DstUV++ = SrcV[u] << 6;
+			uint16* DstUVLine = (uint16*)DstUV;
+			const uint16* SrcU16 = (const uint16*)SrcU;
+			const uint16* SrcV16 = (const uint16*)SrcV;
+			for (int32 u = 0; u < w / 2; ++u)
+			{
+				*DstUVLine++ = SrcU16[u];
+				*DstUVLine++ = SrcV16[u];
+			}
+			SrcU += PitchU;
+			SrcV += PitchV;
+			DstUV += DstPitch;
 		}
-		SrcU += PitchU;
-		SrcV += PitchV;
-		DstUV += padUV;
 	}
+
+#if ELECTRA_MEDIAGPUBUFFER_DX12
+	if (InNewOutput->GPUBuffer.Resource.IsValid())
+	{
+		InNewOutput->GPUBuffer.Resource->Unmap(0, nullptr);
+		// To be compatible with implementations that might do the copy into the resource async, we also signal a fence
+		// (strictly speaking we would not need to as this is all 100% synchronous and done before the GPU ever attempts to read from the resource)
+		InNewOutput->GPUBuffer.Fence->Signal(InNewOutput->GPUBuffer.FenceValue);
+	}
+#endif
+
 	return true;
 }
diff --git a/Engine/Plugins/Media/VPxDecoderElectra/Source/VPxDecoderElectra/VPxDecoderElectra.Build.cs b/Engine/Plugins/Media/VPxDecoderElectra/Source/VPxDecoderElectra/VPxDecoderElectra.Build.cs
--- a/Engine/Plugins/Media/VPxDecoderElectra/Source/VPxDecoderElectra/VPxDecoderElectra.Build.cs
+++ b/Engine/Plugins/Media/VPxDecoderElectra/Source/VPxDecoderElectra/VPxDecoderElectra.Build.cs
@@ -1,5 +1,7 @@
 // Copyright Epic Games, Inc. All Rights Reserved.
 
+using System.IO;
+
 namespace UnrealBuildTool.Rules
 {
 	public class VPxDecoderElectra : ModuleRules
@@ -13,6 +15,19 @@ namespace UnrealBuildTool.Rules
 					"ElectraDecoders",
 					"LibVpx"
                 });
+
+			if (Target.Platform.IsInGroup(UnrealPlatformGroup.Windows))
+			{
+				PublicDependencyModuleNames.Add("DirectX");
+				AddEngineThirdPartyPrivateStaticDependencies(Target, "DX12");
+
+				if (Target.Platform.IsInGroup(UnrealPlatformGroup.Windows) && Target.WindowsPlatform.Architecture != UnrealArch.Arm64)
+				{
+					PublicAdditionalLibraries.AddRange(new string[] {
+						Path.Combine(Target.WindowsPlatform.DirectXLibDir, "dxerr.lib"),
+					});
+				}
+			}
 		}
 	}
 }
diff --git a/Engine/Source/Runtime/MediaAssets/Private/Misc/MediaTextureResource.cpp b/Engine/Source/Runtime/MediaAssets/Private/Misc/MediaTextureResource.cpp
--- a/Engine/Source/Runtime/MediaAssets/Private/Misc/MediaTextureResource.cpp
+++ b/Engine/Source/Runtime/MediaAssets/Private/Misc/MediaTextureResource.cpp
@@ -320,7 +320,9 @@ namespace MediaTextureResourceHelpers
 		EMediaTextureSampleFormat Fmt = Sample->GetFormat();
 		return Fmt == EMediaTextureSampleFormat::CharBGRA ||
 			Fmt == EMediaTextureSampleFormat::CharRGBA ||
-			Fmt == EMediaTextureSampleFormat::CharBMP;
+			Fmt == EMediaTextureSampleFormat::CharBMP ||
+			Fmt == EMediaTextureSampleFormat::DXT1 ||
+			Fmt == EMediaTextureSampleFormat::DXT5;
 	}
 
 } //namespace
@@ -1099,7 +1101,7 @@ void FMediaTextureResource::GetColorSpaceConversionMatrixForSample(const TShared
 						GraphicsPSOInit.BoundShaderState.PixelShaderRHI = ConvertShader.GetPixelShader();
 						SetGraphicsPipelineState(CommandList, GraphicsPSOInit, 0);
 						FIntPoint TexDim = InputTexture->GetSizeXY();
-						TempSRV0 = CommandList.CreateShaderResourceView(InputTexture, 0, 1, PF_G8);		// note: the types of the views select the correct planes (offset) "magically"
+						TempSRV0 = CommandList.CreateShaderResourceView(InputTexture, 0, 1, PF_G8);								// note: RHI does provide "magic" to select Y vs. UV planes based on the pixel format (D3D/DXGI)
 						TempSRV1 = CommandList.CreateShaderResourceView(InputTexture, 0, 1, PF_R8G8);
 						SetShaderParametersLegacyPS(CommandList, ConvertShader, TexDim, TempSRV0, TempSRV1, OutputDim, YUVMtx, Sample->GetEncodingType(), ColorSpaceMtx, SampleFormat == EMediaTextureSampleFormat::CharNV21);
 					}
@@ -1127,7 +1129,7 @@ void FMediaTextureResource::GetColorSpaceConversionMatrixForSample(const TShared
 						GraphicsPSOInit.BoundShaderState.PixelShaderRHI = ConvertShader.GetPixelShader();
 						SetGraphicsPipelineState(CommandList, GraphicsPSOInit, 0);
 
-						FShaderResourceViewRHIRef Y_SRV = CommandList.CreateShaderResourceView(InputTexture, 0, 1, PF_G16);		// note: the types of the views select the correct planes (offset) "magically"
+						FShaderResourceViewRHIRef Y_SRV = CommandList.CreateShaderResourceView(InputTexture, 0, 1, PF_G16);		// note: RHI does provide "magic" to select Y vs. UV planes based on the pixel format (D3D/DXGI)
 						FShaderResourceViewRHIRef UV_SRV = CommandList.CreateShaderResourceView(InputTexture, 0, 1, PF_G16R16);
 						SetShaderParametersLegacyPS(CommandList, ConvertShader, TexDim, Y_SRV, UV_SRV, OutputDim, YUVMtx, ColorSpaceMtx, Sample->GetEncodingType());
 					}
@@ -1141,7 +1143,7 @@ void FMediaTextureResource::GetColorSpaceConversionMatrixForSample(const TShared
 					}
 				}
 				break;
-
+ 
 				// Various YCbCr 422 formats
 				case EMediaTextureSampleFormat::CharYUY2:		// Y0CbY1Cr
 				case EMediaTextureSampleFormat::CharYVYU:		// Y0CrY1Cb
