Disabling some static code analysis warnings in Electra's H264/5 decoder code that trigger incorrectly
#rnx

diff --git a/Engine/Plugins/Media/ElectraCodecs/Source/ElectraDecoders/Private/DX/VideoDecoderH264_DX.cpp b/Engine/Plugins/Media/ElectraCodecs/Source/ElectraDecoders/Private/DX/VideoDecoderH264_DX.cpp
--- a/Engine/Plugins/Media/ElectraCodecs/Source/ElectraDecoders/Private/DX/VideoDecoderH264_DX.cpp
+++ b/Engine/Plugins/Media/ElectraCodecs/Source/ElectraDecoders/Private/DX/VideoDecoderH264_DX.cpp
@@ -106,7 +106,7 @@ public:
 			case EElectraDecoderPlatformOutputHandleType::ImageBuffers:
 				if (!MFSample.IsValid())
 				{
-					if (InTypeOfHandle == EElectraDecoderPlatformOutputHandleType::ImageBuffers)
+					if (InTypeOfHandle == EElectraDecoderPlatformOutputHandleType::ImageBuffers)	//-V547
 					{
 						return static_cast<IElectraDecoderVideoOutputImageBuffers*>(const_cast<FElectraVideoDecoderOutputH264_DX*>(this));
 					}
diff --git a/Engine/Plugins/Media/ElectraCodecs/Source/ElectraDecoders/Private/DX/VideoDecoderH265_DX.cpp b/Engine/Plugins/Media/ElectraCodecs/Source/ElectraDecoders/Private/DX/VideoDecoderH265_DX.cpp
--- a/Engine/Plugins/Media/ElectraCodecs/Source/ElectraDecoders/Private/DX/VideoDecoderH265_DX.cpp
+++ b/Engine/Plugins/Media/ElectraCodecs/Source/ElectraDecoders/Private/DX/VideoDecoderH265_DX.cpp
@@ -122,7 +122,7 @@ public:
 		case EElectraDecoderPlatformOutputHandleType::ImageBuffers:
 			if (!MFSample.IsValid())
 			{
-				if (InTypeOfHandle == EElectraDecoderPlatformOutputHandleType::ImageBuffers)
+				if (InTypeOfHandle == EElectraDecoderPlatformOutputHandleType::ImageBuffers)	//-V547
 				{
 					return static_cast<IElectraDecoderVideoOutputImageBuffers*>(const_cast<FElectraVideoDecoderOutputH265_DX*>(this));
 				}
