Fixing up various Electra issues when playing high FPS videos with new DX12 output path code

diff --git a/Engine/Plugins/Media/ElectraCodecs/Source/ElectraDecoders/Private/DX/VideoDecoderH264_DX.cpp b/Engine/Plugins/Media/ElectraCodecs/Source/ElectraDecoders/Private/DX/VideoDecoderH264_DX.cpp
--- a/Engine/Plugins/Media/ElectraCodecs/Source/ElectraDecoders/Private/DX/VideoDecoderH264_DX.cpp
+++ b/Engine/Plugins/Media/ElectraCodecs/Source/ElectraDecoders/Private/DX/VideoDecoderH264_DX.cpp
@@ -1166,7 +1166,10 @@ bool FElectraVideoDecoderH264_DX::ConvertDecoderOutput()
 
 		// Request resource and fence...
 		uint32 BufferPitch;
-		D3D12ResourcePool->AllocateOutputDataAsBuffer(NewOutput->Buffer, BufferPitch);
+		if (!D3D12ResourcePool->AllocateOutputDataAsBuffer(NewOutput->Buffer, BufferPitch))
+		{
+			return false;
+		}
 
 		NewOutput->Pitch = BufferPitch;
 
diff --git a/Engine/Plugins/Media/ElectraCodecs/Source/ElectraDecoders/Private/DX/VideoDecoderH265_DX.cpp b/Engine/Plugins/Media/ElectraCodecs/Source/ElectraDecoders/Private/DX/VideoDecoderH265_DX.cpp
--- a/Engine/Plugins/Media/ElectraCodecs/Source/ElectraDecoders/Private/DX/VideoDecoderH265_DX.cpp
+++ b/Engine/Plugins/Media/ElectraCodecs/Source/ElectraDecoders/Private/DX/VideoDecoderH265_DX.cpp
@@ -1222,7 +1222,10 @@ bool FElectraVideoDecoderH265_DX::ConvertDecoderOutput()
 
 		// Request resource and fence...
 		uint32 BufferPitch;
-		D3D12ResourcePool->AllocateOutputDataAsBuffer(NewOutput->Buffer, BufferPitch);
+		if (!D3D12ResourcePool->AllocateOutputDataAsBuffer(NewOutput->Buffer, BufferPitch))
+		{
+			return false;
+		}
 
 		NewOutput->Pitch = BufferPitch;
 
diff --git a/Engine/Plugins/Media/ElectraCodecs/Source/ElectraDecoders/Public/Windows/WindowsElectraDecoderGPUBufferHelpers.h b/Engine/Plugins/Media/ElectraCodecs/Source/ElectraDecoders/Public/Windows/WindowsElectraDecoderGPUBufferHelpers.h
--- a/Engine/Plugins/Media/ElectraCodecs/Source/ElectraDecoders/Public/Windows/WindowsElectraDecoderGPUBufferHelpers.h
+++ b/Engine/Plugins/Media/ElectraCodecs/Source/ElectraDecoders/Public/Windows/WindowsElectraDecoderGPUBufferHelpers.h
@@ -68,26 +68,29 @@
 /*********************************************************************************************************************/
 /*********************************************************************************************************************/
 
+class FElectraMediaDecoderOutputBufferPool_DX12;
+
 // Simple DX12 heap and fence manager for use with decoder output buffers (both upload heaps or default GPU heaps)
-class FElectraMediaDecoderOutputBufferPool_DX12 : public TSharedFromThis<FElectraMediaDecoderOutputBufferPool_DX12, ESPMode::ThreadSafe>
+class FElectraMediaDecoderOutputBufferPoolBlock_DX12
 {
 	struct FResourceInfo
 	{
-		FResourceInfo(TWeakPtr<FElectraMediaDecoderOutputBufferPool_DX12> InPool, uint32 InBufferIdx, TRefCountPtr<ID3D12Heap> InHeap) : Pool(InPool), BufferIdx(InBufferIdx), Heap(InHeap) {}
+		FResourceInfo(TWeakPtr<FElectraMediaDecoderOutputBufferPool_DX12> InPool, uint32 InBlockIdx, uint32 InBufferIdx, TRefCountPtr<ID3D12Heap> InHeap) : Pool(InPool), BlockIdx(InBlockIdx), BufferIdx(InBufferIdx), Heap(InHeap) {}
 
 		TWeakPtr<FElectraMediaDecoderOutputBufferPool_DX12> Pool;
+		uint32 BlockIdx;
 		uint32 BufferIdx;
 		TRefCountPtr<ID3D12Heap> Heap;
 	};
 
 public:
 	// Create instance for use with buffers (usually for uploading data)
-	FElectraMediaDecoderOutputBufferPool_DX12(TRefCountPtr<ID3D12Device> InD3D12Device, uint32 InMaxNumBuffers, uint32 Width, uint32 Height, uint32 BytesPerPixel, D3D12_HEAP_TYPE InD3D12HeapType = D3D12_HEAP_TYPE_UPLOAD)
+	FElectraMediaDecoderOutputBufferPoolBlock_DX12(uint32 InBlockIdx, TRefCountPtr<ID3D12Device> InD3D12Device, uint32 InMaxNumBuffers, uint32 Width, uint32 Height, uint32 BytesPerPixel, D3D12_HEAP_TYPE InD3D12HeapType = D3D12_HEAP_TYPE_UPLOAD)
 		: D3D12Device(InD3D12Device)
 		, D3D12HeapType(InD3D12HeapType)
 		, MaxNumBuffers(InMaxNumBuffers)
 		, FreeMask((1 << InMaxNumBuffers) - 1)
-		, LastFenceValue(0)
+		, BlockIdx(InBlockIdx)
 	{
 		check(InMaxNumBuffers <= 32);
 
@@ -99,12 +102,12 @@
 	}
 
 	// Create instance for use with textures (usually to receive any upload buffers or internal decoder texture data)
-	FElectraMediaDecoderOutputBufferPool_DX12(TRefCountPtr<ID3D12Device> InD3D12Device, uint32 InMaxNumBuffers, uint32 Width, uint32 Height, DXGI_FORMAT PixFmt, D3D12_HEAP_TYPE InD3D12HeapType = D3D12_HEAP_TYPE_DEFAULT)
+	FElectraMediaDecoderOutputBufferPoolBlock_DX12(uint32 InBlockIdx, TRefCountPtr<ID3D12Device> InD3D12Device, uint32 InMaxNumBuffers, uint32 Width, uint32 Height, DXGI_FORMAT PixFmt, D3D12_HEAP_TYPE InD3D12HeapType = D3D12_HEAP_TYPE_DEFAULT)
 		: D3D12Device(InD3D12Device)
 		, D3D12HeapType(InD3D12HeapType)
 		, MaxNumBuffers(InMaxNumBuffers)
 		, FreeMask((1 << InMaxNumBuffers) - 1)
-		, LastFenceValue(0)
+		, BlockIdx(InBlockIdx)
 	{
 		check(InMaxNumBuffers <= 32);
 
@@ -127,7 +130,7 @@
 		InitCommon();
 	}
 
-	~FElectraMediaDecoderOutputBufferPool_DX12()
+	~FElectraMediaDecoderOutputBufferPoolBlock_DX12()
 	{
 	}
 
@@ -145,36 +148,10 @@
 		return FPlatformAtomics::AtomicRead((const int32*)&FreeMask) != 0;
 	}
 
-	class FOutputData
-	{
-	public:
-		~FOutputData()
-		{
-			check(ReadyToDestroy());
-		}
-
-		TRefCountPtr<ID3D12Resource> Resource;
-		TRefCountPtr<ID3D12Fence> Fence;
-		uint64 FenceValue;
-
-		bool ReadyToDestroy() const
-		{
-			bool bOk = true;
-#if !UE_BUILD_SHIPPING
-			if (Resource.IsValid())
-			{
-				bOk = Resource->AddRef() > 1 || Fence->GetCompletedValue() >= FenceValue;
-				Resource->Release();
-			}
-#endif
-			return bOk;
-		}
-	};
-
 	// Allocate a buffer resource and return suitable sync fence data
-	bool AllocateOutputDataAsBuffer(FOutputData& OutData, uint32& OutPitch)
+	TRefCountPtr<ID3D12Resource> AllocateOutputDataAsBuffer(TWeakPtr<FElectraMediaDecoderOutputBufferPool_DX12> InPoolWeakRef, uint32& OutPitch)
 	{
-		OutData.Resource = AllocateBuffer(OutPitch, [NumBytes = BufferSize](D3D12_RESOURCE_DESC& Desc)
+		TRefCountPtr<ID3D12Resource> Resource = AllocateBuffer(InPoolWeakRef, OutPitch, [NumBytes = BufferSize](D3D12_RESOURCE_DESC& Desc)
 		{
 			Desc.MipLevels = 1;
 			Desc.Format = DXGI_FORMAT_UNKNOWN;
@@ -187,15 +164,14 @@
 			Desc.Layout = D3D12_TEXTURE_LAYOUT_ROW_MAJOR;
 			Desc.Dimension = D3D12_RESOURCE_DIMENSION_BUFFER;
 		});
-		OutData.Fence = GetUpdatedBufferFence(OutData.FenceValue);
-		return OutData.Resource.IsValid();
+		return Resource;
 	}
 
 	// Allocate a texture resource and return suitable sync fence data
-	bool AllocateOutputDataAsTexture(FOutputData& OutData, uint32 InWidth, uint32 InHeight, DXGI_FORMAT InPixFmt)
+	TRefCountPtr<ID3D12Resource> AllocateOutputDataAsTexture(TWeakPtr<FElectraMediaDecoderOutputBufferPool_DX12> InPoolWeakRef, uint32 InWidth, uint32 InHeight, DXGI_FORMAT InPixFmt)
 	{
 		uint32 OutPitch;
-		OutData.Resource = AllocateBuffer(OutPitch, [InWidth, InHeight, InPixFmt](D3D12_RESOURCE_DESC& Desc)
+		TRefCountPtr<ID3D12Resource> Resource = AllocateBuffer(InPoolWeakRef, OutPitch, [InWidth, InHeight, InPixFmt](D3D12_RESOURCE_DESC& Desc)
 		{
 			Desc.MipLevels = 1;
 			Desc.Format = InPixFmt;
@@ -208,42 +184,7 @@
 			Desc.Layout = D3D12_TEXTURE_LAYOUT_UNKNOWN;
 			Desc.Dimension = D3D12_RESOURCE_DIMENSION_TEXTURE2D;
 		});
-		OutData.Fence = GetUpdatedBufferFence(OutData.FenceValue);
-		return OutData.Resource.IsValid();
+		return Resource;
 	}
 
-	TRefCountPtr<ID3D12Fence> GetUpdatedBufferFence(uint64& FenceValue)
-	{
-		FenceValue = ++LastFenceValue;
-		return D3D12BufferFence;
-	}
-
-	// Helper function to copy simple, linear texture data between differently picthed buffers
-	static void CopyWithPitchAdjust(uint8* Dst, uint32 DstPitch, const uint8* Src, uint32 SrcPitch, uint32 NumRows)
-	{
-		if (DstPitch != SrcPitch)
-		{
-			for (uint32 Y = NumRows; Y > 0; --Y)
-			{
-				FMemory::Memcpy(Dst, Src, SrcPitch);
-				Src += SrcPitch;
-				Dst += DstPitch;
-			}
-		}
-		else
-		{
-			uint32 NumBytes = SrcPitch * NumRows;
-#if ALLOW_USE_OF_SSEMEMCPY
-			if ((NumBytes & 0x7f) == 0)
-			{
-				SSE2MemCpy(Dst, Src, NumBytes);
-			}
-			else
-#endif
-			{
-				FMemory::Memcpy(Dst, Src, NumBytes);
-			}
-		}
-	}
-
 private:
@@ -263,15 +204,9 @@
 #if !UE_BUILD_SHIPPING
 		D3D12OutputHeap->SetName(TEXT("ElectraOutputBufferPoolHeap"));
 #endif
-
-		Res = D3D12Device->CreateFence(0, D3D12_FENCE_FLAG_NONE, IID_PPV_ARGS(D3D12BufferFence.GetInitReference()));
-		check(SUCCEEDED(Res));
-#if !UE_BUILD_SHIPPING
-		D3D12BufferFence->SetName(TEXT("ElectraOutputBufferPoolFence"));
-#endif
 	}
 
-	TRefCountPtr<ID3D12Resource> AllocateBuffer(uint32& OutBufferPitch, TFunction<void(D3D12_RESOURCE_DESC& Desc)> && InitializeDesc)
+	TRefCountPtr<ID3D12Resource> AllocateBuffer(TWeakPtr<FElectraMediaDecoderOutputBufferPool_DX12> InPoolWeakRef, uint32& OutBufferPitch, TFunction<void(D3D12_RESOURCE_DESC& Desc)> && InitializeDesc)
 	{
 		TRefCountPtr<ID3D12Resource> Resource;
 
@@ -312,7 +247,7 @@
 
 				// note: we keep a reference to the heap in our context data for the destruction callback, so we can ensure the heap lives as long as there are placed resources in it
 				UINT CallbackID;
-				Res = Notifier->RegisterDestructionCallback(ResourceDestructionCallback, new FResourceInfo(AsWeak(), BufferIdx, D3D12OutputHeap), &CallbackID);
+				Res = Notifier->RegisterDestructionCallback(ResourceDestructionCallback, new FResourceInfo(InPoolWeakRef, BlockIdx, BufferIdx, D3D12OutputHeap), &CallbackID);
 				check(SUCCEEDED(Res));
 			}
 			else
@@ -335,30 +270,246 @@
 		} while (OldFreeMask != FPlatformAtomics::InterlockedCompareExchange((int32*)&FreeMask, NewFreeMask, OldFreeMask));
 	}
 
-	static void ResourceDestructionCallback(void* Context)
-	{
-		auto ResourceInfo = reinterpret_cast<FResourceInfo*>(Context);
-
-		if (auto Pool = ResourceInfo->Pool.Pin())
-		{
-			Pool->FreeBuffer(ResourceInfo->BufferIdx);
-		}
-
-		// Get rid of our resource tracking info block...
-		// (this will also release the ref to the heap the resource was on)
-		delete ResourceInfo;
-	}
+	static void ResourceDestructionCallback(void* Context);
 
 	TRefCountPtr<ID3D12Device> D3D12Device;
 	D3D12_HEAP_TYPE D3D12HeapType;
 	TRefCountPtr<ID3D12Heap> D3D12OutputHeap;
-	TRefCountPtr<ID3D12Fence> D3D12BufferFence;
 	uint32 MaxNumBuffers;
 	uint32 BufferSize;
 	uint32 BufferPitch;
 	uint32 FreeMask;
-	uint64 LastFenceValue;
+	uint32 BlockIdx;
 };
 
 
+class FElectraMediaDecoderOutputBufferPool_DX12 : public TSharedFromThis<FElectraMediaDecoderOutputBufferPool_DX12, ESPMode::ThreadSafe>
+{
+	enum {
+		kExpectedMaxNumBlocks = 8,
+	};
+	
+public:
+	// Create instance for use with buffers (usually for uploading data)
+	FElectraMediaDecoderOutputBufferPool_DX12(TRefCountPtr<ID3D12Device> InD3D12Device, uint32 InMaxNumBuffers, uint32 InWidth, uint32 InHeight, uint32 InBytesPerPixel, D3D12_HEAP_TYPE InD3D12HeapType = D3D12_HEAP_TYPE_UPLOAD)
+		: D3D12Device(InD3D12Device)
+		, LastFenceValue(0)
+		, D3D12HeapType(InD3D12HeapType)
+		, Width(InWidth)
+		, Height(InHeight)
+		, PixFmt(DXGI_FORMAT_UNKNOWN)
+		, BytesPerPixel(InBytesPerPixel)
+	{
+		Blocks.Reserve(kExpectedMaxNumBlocks);
+		Blocks.Emplace(new FElectraMediaDecoderOutputBufferPoolBlock_DX12(0, D3D12Device, InMaxNumBuffers, Width, Height, BytesPerPixel, D3D12HeapType));
+		InitCommon(InMaxNumBuffers);
+	}
+
+	// Create instance for use with textures (usually to receive any upload buffers or internal decoder texture data)
+	FElectraMediaDecoderOutputBufferPool_DX12(TRefCountPtr<ID3D12Device> InD3D12Device, uint32 InMaxNumBuffers, uint32 InWidth, uint32 InHeight, DXGI_FORMAT InPixFmt, D3D12_HEAP_TYPE InD3D12HeapType = D3D12_HEAP_TYPE_DEFAULT)
+		: D3D12Device(InD3D12Device)
+		, LastFenceValue(0)
+		, D3D12HeapType(InD3D12HeapType)
+		, Width(InWidth)
+		, Height(InHeight)
+		, PixFmt(InPixFmt)
+		, BytesPerPixel(0)
+	{
+		Blocks.Reserve(kExpectedMaxNumBlocks);
+		Blocks.Emplace(new FElectraMediaDecoderOutputBufferPoolBlock_DX12(0, D3D12Device, InMaxNumBuffers, Width, Height, PixFmt, D3D12HeapType));
+		InitCommon(InMaxNumBuffers);
+	}
+
+	~FElectraMediaDecoderOutputBufferPool_DX12()
+	{
+		for (FElectraMediaDecoderOutputBufferPoolBlock_DX12* Block : Blocks)
+		{
+			delete Block;
+		}
+	}
+
+	// Check if the current setup is compatible with the new parameters
+	bool IsCompatibleAsBuffer(uint32 InMaxNumBuffers, uint32 InWidth, uint32 InHeight, uint32 InBytesPerPixel) const
+	{
+		return Blocks[0]->IsCompatibleAsBuffer(InMaxNumBuffers, InWidth, InHeight, InBytesPerPixel);
+	}
+
+	// Check if a buffer is available
+	bool BufferAvailable() const
+	{
+		for (FElectraMediaDecoderOutputBufferPoolBlock_DX12* Block : Blocks)
+		{
+			if (Block->BufferAvailable())
+			{
+				return true;
+			}
+		}
+		return false;
+	}
+
+	class FOutputData
+	{
+	public:
+		~FOutputData()
+		{
+			check(ReadyToDestroy());
+		}
+
+		TRefCountPtr<ID3D12Resource> Resource;
+		TRefCountPtr<ID3D12Fence> Fence;
+		uint64 FenceValue;
+
+		bool ReadyToDestroy() const
+		{
+			bool bOk = true;
+#if !UE_BUILD_SHIPPING
+			if (Resource.IsValid())
+			{
+				bOk = Resource->AddRef() > 1 || Fence->GetCompletedValue() >= FenceValue;
+				Resource->Release();
+			}
+#endif
+			return bOk;
+		}
+	};
+
+	bool AllocateOutputDataAsBuffer(FOutputData& OutData, uint32& OutPitch)
+	{
+		check(BytesPerPixel != 0);
+
+		OutData.Resource = nullptr;
+		for (FElectraMediaDecoderOutputBufferPoolBlock_DX12* Block : Blocks)
+		{
+			if (Block->BufferAvailable())
+			{
+				OutData.Resource = Block->AllocateOutputDataAsBuffer(AsWeak(), OutPitch);
+				if (OutData.Resource)
+				{
+					break;
+				}
+			}
+		}
+
+		if (!OutData.Resource.IsValid())
+		{
+			Blocks.Emplace(new FElectraMediaDecoderOutputBufferPoolBlock_DX12(Blocks.Num(), D3D12Device, MaxNumBuffersAddBlocks, Width, Height, BytesPerPixel, D3D12HeapType));
+			OutData.Resource = Blocks.Last()->AllocateOutputDataAsBuffer(AsWeak(), OutPitch);
+		}
+
+		if (OutData.Resource.IsValid())
+		{
+			OutData.Fence = GetUpdatedBufferFence(OutData.FenceValue);
+			return true;
+		}
+		return false;
+	}
+
+	bool AllocateOutputDataAsTexture(FOutputData& OutData, uint32 InWidth, uint32 InHeight, DXGI_FORMAT InPixFmt)
+	{
+		check(PixFmt != DXGI_FORMAT_UNKNOWN);
+
+		OutData.Resource = nullptr;
+		for (FElectraMediaDecoderOutputBufferPoolBlock_DX12* Block : Blocks)
+		{
+			if (Block->BufferAvailable())
+			{
+				OutData.Resource = Block->AllocateOutputDataAsTexture(AsWeak(), InWidth, InHeight, InPixFmt);
+				if (OutData.Resource)
+				{
+					break;
+				}
+			}
+		}
+
+		if (!OutData.Resource.IsValid())
+		{
+			Blocks.Emplace(new FElectraMediaDecoderOutputBufferPoolBlock_DX12(Blocks.Num(), D3D12Device, MaxNumBuffersAddBlocks, Width, Height, PixFmt, D3D12HeapType));
+			OutData.Resource = Blocks.Last()->AllocateOutputDataAsTexture(AsWeak(), InWidth, InHeight, InPixFmt);
+		}
+
+		if (OutData.Resource.IsValid())
+		{
+			OutData.Fence = GetUpdatedBufferFence(OutData.FenceValue);
+			return true;
+		}
+		return false;
+	}
+
+	TRefCountPtr<ID3D12Fence> GetUpdatedBufferFence(uint64& FenceValue)
+	{
+		FenceValue = ++LastFenceValue;
+		return D3D12BufferFence;
+	}
+
+	// Helper function to copy simple, linear texture data between differently picthed buffers
+	static void CopyWithPitchAdjust(uint8* Dst, uint32 DstPitch, const uint8* Src, uint32 SrcPitch, uint32 NumRows)
+	{
+		if (DstPitch != SrcPitch)
+		{
+			for (uint32 Y = NumRows; Y > 0; --Y)
+			{
+				FMemory::Memcpy(Dst, Src, SrcPitch);
+				Src += SrcPitch;
+				Dst += DstPitch;
+			}
+		}
+		else
+		{
+			uint32 NumBytes = SrcPitch * NumRows;
+#if ALLOW_USE_OF_SSEMEMCPY
+			if ((NumBytes & 0x7f) == 0)
+			{
+				SSE2MemCpy(Dst, Src, NumBytes);
+			}
+			else
+#endif
+			{
+				FMemory::Memcpy(Dst, Src, NumBytes);
+			}
+		}
+	}
+
+private:
+	friend class FElectraMediaDecoderOutputBufferPoolBlock_DX12;
+
+	void InitCommon(uint32 InMaxNumBuffers)
+	{
+		MaxNumBuffersAddBlocks = (InMaxNumBuffers > 1) ? (InMaxNumBuffers >> 1) : 1;
+
+		HRESULT Res = D3D12Device->CreateFence(0, D3D12_FENCE_FLAG_NONE, IID_PPV_ARGS(D3D12BufferFence.GetInitReference()));
+		check(SUCCEEDED(Res));
+#if !UE_BUILD_SHIPPING
+		D3D12BufferFence->SetName(TEXT("ElectraOutputBufferPoolFence"));
+#endif
+	}
+
+	TArray<FElectraMediaDecoderOutputBufferPoolBlock_DX12*> Blocks;
+	
+	TRefCountPtr<ID3D12Device> D3D12Device;
+	TRefCountPtr<ID3D12Fence> D3D12BufferFence;
+	uint64 LastFenceValue;
+	D3D12_HEAP_TYPE D3D12HeapType;
+	uint32 Width;
+	uint32 Height;
+	DXGI_FORMAT PixFmt;
+	uint32 BytesPerPixel;
+	uint32 MaxNumBuffersAddBlocks;
+	};
+
+
+inline void FElectraMediaDecoderOutputBufferPoolBlock_DX12::ResourceDestructionCallback(void* Context)
+{
+	auto ResourceInfo = reinterpret_cast<FResourceInfo*>(Context);
+
+	if (auto Pool = ResourceInfo->Pool.Pin())
+	{
+		Pool->Blocks[ResourceInfo->BlockIdx]->FreeBuffer(ResourceInfo->BufferIdx);
+	}
+
+	// Get rid of our resource tracking info block...
+	// (this will also release the ref to the heap the resource was on)
+	delete ResourceInfo;
+}
+
+
 #include "CommonElectraDecoderGPUBufferHelpers.h"
diff --git a/Engine/Plugins/Media/ElectraPlayer/Source/ElectraPlayerRuntime/Private/Runtime/Decoder/Windows/ElectraVideoDecoder_PC.cpp b/Engine/Plugins/Media/ElectraPlayer/Source/ElectraPlayerRuntime/Private/Runtime/Decoder/Windows/ElectraVideoDecoder_PC.cpp
--- a/Engine/Plugins/Media/ElectraPlayer/Source/ElectraPlayerRuntime/Private/Runtime/Decoder/Windows/ElectraVideoDecoder_PC.cpp
+++ b/Engine/Plugins/Media/ElectraPlayer/Source/ElectraPlayerRuntime/Private/Runtime/Decoder/Windows/ElectraVideoDecoder_PC.cpp
@@ -38,16 +38,15 @@ FElectraPlayerVideoDecoderOutputPC::EOutputType FElectraPlayerVideoDecoderOutput
 
 const TArray<uint8>& FElectraPlayerVideoDecoderOutputPC::GetBuffer() const
 {
-	check(OutputType == EOutputType::SoftwareWin8Plus || OutputType == EOutputType::SoftwareWin7 || OutputType == EOutputType::HardwareDX9_DX12);
-	if (Buffer.IsValid())
+	if ((OutputType == EOutputType::SoftwareWin8Plus || OutputType == EOutputType::SoftwareWin7 || OutputType == EOutputType::HardwareDX9_DX12))
 	{
-		return *Buffer;
-	}
-	else
-	{
-		static TArray<uint8> Empty;
-		return Empty;
+		if (Buffer.IsValid())
+		{
+			return *Buffer;
+		}
 	}
+	static TArray<uint8> Empty;
+	return Empty;
 }
 
 uint32 FElectraPlayerVideoDecoderOutputPC::GetStride() const
@@ -99,6 +98,11 @@ void FElectraPlayerVideoDecoderOutputPC::InitializeWithBuffer(const void* InBuff
 
 	SampleDim = Dim;
 	Stride = InStride;
+
+	Texture = nullptr;
+	SharedTexture = nullptr;
+	TextureDX12 = nullptr;
+	D3DFence = nullptr;
 }
 
 void FElectraPlayerVideoDecoderOutputPC::InitializeWithBuffer(TSharedPtr<TArray<uint8>, ESPMode::ThreadSafe> InBuffer, uint32 InStride, FIntPoint Dim, TSharedPtr<Electra::FParamDict, ESPMode::ThreadSafe> InParamDict)
@@ -111,11 +115,19 @@ void FElectraPlayerVideoDecoderOutputPC::InitializeWithBuffer(TSharedPtr<TArray<
 
 	SampleDim = Dim;
 	Stride = InStride;
+
+	Texture = nullptr;
+	SharedTexture = nullptr;
+	TextureDX12 = nullptr;
+	D3DFence = nullptr;
 }
 
 void FElectraPlayerVideoDecoderOutputPC::InitializeWithResource(const TRefCountPtr<ID3D12Device>& InD3D12Device, const TRefCountPtr<ID3D12Resource> InResource, uint32 ResourcePitch, const FElectraDecoderOutputSync& OutputSync, const FIntPoint& InOutputDim, TSharedPtr<Electra::FParamDict, ESPMode::ThreadSafe> InParamDict, TWeakPtr<Electra::IVideoDecoderResourceDelegate, ESPMode::ThreadSafe> InResourceDelegate,
-																uint32 MaxWidth, uint32 MaxHeight, uint32 MaxOutputBuffers)
+																TSharedPtr<FElectraMediaDecoderOutputBufferPool_DX12>& InOutD3D12ResourcePool, uint32 MaxWidth, uint32 MaxHeight, uint32 MaxOutputBuffers)
 {
+	// We must not allow re-initialization of used instances
+	check(!D3DFence.IsValid());
+
 	// General initialization
 	FVideoDecoderOutput::Initialize(InParamDict);
 
@@ -124,6 +136,12 @@ void FElectraPlayerVideoDecoderOutputPC::InitializeWithResource(const TRefCountP
 
 	SampleDim = InOutputDim;
 	OutputType = EOutputType::Hardware_DX;
+	if (!D3D12ResourcePool.IsValid() && InOutD3D12ResourcePool.IsValid())
+	{
+		D3D12ResourcePool = InOutD3D12ResourcePool;
+	}
+
+	Buffer.Reset();
 
 	HRESULT Res;
 
@@ -211,6 +229,7 @@ void FElectraPlayerVideoDecoderOutputPC::InitializeWithResource(const TRefCountP
 				UE_LOG(LogElectraPlayer, Error, TEXT("Could not allocate texture resource heap!"));
 				return;
 			}
+			InOutD3D12ResourcePool = D3D12ResourcePool;
 		}
 		
 		FElectraMediaDecoderOutputBufferPool_DX12::FOutputData OutputData;
@@ -332,29 +351,6 @@ void FElectraPlayerVideoDecoderOutputPC::InitializeWithResource(const TRefCountP
 
 	// - - - - - - - - - - - - - - - - - - - - - - - - - - - - - - - - - - - - - - - - - - - - - - - - - - - - - - - -
 
-	// Get all synchronization info we need
-	// (we could have a fence or a DX12 decoder synchronization object)
-
-	TRefCountPtr<ID3D12Fence> ResourceFence;
-#if ALLOW_MFSAMPLE_WITH_DX12
-	TRefCountPtr<IMFD3D12SynchronizationObjectCommands> DecoderSync;
-#else
-	TRefCountPtr<IUnknown> DecoderSync;
-#endif
-
-	if (OutputSync.Sync.IsValid())
-	{
-#if ALLOW_MFSAMPLE_WITH_DX12
-		if (OutputSync.Sync->QueryInterface(__uuidof(IMFD3D12SynchronizationObjectCommands), (void**)DecoderSync.GetInitReference()) != S_OK)
-#endif // ALLOW_MFSAMPLE_WITH_DX12
-		{
-			OutputSync.Sync->QueryInterface(__uuidof(ID3D12Fence), (void**)ResourceFence.GetInitReference());
-		}
-		check(ResourceFence.IsValid() || DecoderSync.IsValid());
-	}
-
-	// - - - - - - - - - - - - - - - - - - - - - - - - - - - - - - - - - - - - - - - - - - - - - - - - - - - - - - - -
-
 	// Can we wait for the decoder to have the data ready using an async CPU job?
 	bool bTriggerOk = false;
 	if (OutputSync.TaskSync.IsValid())
@@ -363,35 +359,40 @@ void FElectraPlayerVideoDecoderOutputPC::InitializeWithResource(const TRefCountP
 		auto ElectraDecoderResourceDelegate = Electra::FElectraDecoderResourceManagerWindows::GetDelegate();
 
 		// Note: we capture "this" as we ensure that this instance only dies once the copy triggered here is actually done, hence ensuring any reference to "this" is done
-		bTriggerOk = ElectraDecoderResourceDelegate->RunCodeAsync([WeakThis=AsWeak(), DecoderSync, ResourceFence, OutputSync, InResourceDelegate]()
+		bTriggerOk = ElectraDecoderResourceDelegate->RunCodeAsync([D3DCmdList=this->D3DCmdList, D3DFence=this->D3DFence, FenceValue=this->FenceValue, OutputSync, ResourceDelegate = InResourceDelegate.Pin()]()
 			{
-				if (auto This = StaticCastSharedPtr<FElectraPlayerVideoDecoderOutputPC, IDecoderOutputPoolable, ESPMode::ThreadSafe>(WeakThis.Pin()))
-				{
-					if (auto ResourceDelegate = InResourceDelegate.Pin())
-					{
-						This->TriggerDataCopy(DecoderSync, ResourceFence, OutputSync, ResourceDelegate.Get());
-					}
-				}
+				TriggerDataCopy(D3DCmdList, D3DFence, FenceValue, OutputSync, ResourceDelegate.Get());
 			}, OutputSync.TaskSync.Get());
 	}
 
 	if (!bTriggerOk)
 	{
 		// We could not run the trigger async. Schedule the copy right away. Any needed synchronization will be done in the copy-queue by the GPU
-		TriggerDataCopy(DecoderSync, ResourceFence, OutputSync, InResourceDelegate.Pin().Get());
+		TriggerDataCopy(D3DCmdList, D3DFence, FenceValue, OutputSync, InResourceDelegate.Pin().Get());
 	}
 }
 
-#if ALLOW_MFSAMPLE_WITH_DX12
-void FElectraPlayerVideoDecoderOutputPC::TriggerDataCopy(TRefCountPtr<IMFD3D12SynchronizationObjectCommands> DecoderSync, TRefCountPtr<ID3D12Fence> ResourceFence, const FElectraDecoderOutputSync& OutputSync, Electra::IVideoDecoderResourceDelegate* InResourceDelegate) const
-#else
-void FElectraPlayerVideoDecoderOutputPC::TriggerDataCopy(TRefCountPtr<IUnknown> DecoderSync, TRefCountPtr<ID3D12Fence> ResourceFence, const FElectraDecoderOutputSync& OutputSync, Electra::IVideoDecoderResourceDelegate* InResourceDelegate) const
-#endif
+void FElectraPlayerVideoDecoderOutputPC::TriggerDataCopy(TRefCountPtr<ID3D12GraphicsCommandList> D3DCmdList, TRefCountPtr<ID3D12Fence> D3DFence, uint64 FenceValue, const FElectraDecoderOutputSync& OutputSync, Electra::IVideoDecoderResourceDelegate* InResourceDelegate)
 {
 	// Trigger copy (this will eventually execute on the submission thread of RHI if running in UE)
 	// (note: we pass in all of FElectraDecoderOutputSync to guarantee any references needed to make the decoder output sync work are passed along, too!)
-	InResourceDelegate->ExecuteCodeWithCopyCommandQueueUsage([CmdList = D3DCmdList, DestFence = D3DFence, DestFenceValue = FenceValue, DecoderSync, ResourceFence, OutputSync](ID3D12CommandQueue* D3DCmdQueue)
+	InResourceDelegate->ExecuteCodeWithCopyCommandQueueUsage([CmdList = D3DCmdList, DestFence = D3DFence, DestFenceValue = FenceValue, OutputSync](ID3D12CommandQueue* D3DCmdQueue)
 		{
+			TRefCountPtr<ID3D12Fence> ResourceFence;
+	#if ALLOW_MFSAMPLE_WITH_DX12
+			TRefCountPtr<IMFD3D12SynchronizationObjectCommands> DecoderSync;
+	#endif
+
+			if (OutputSync.Sync.IsValid())
+			{
+	#if ALLOW_MFSAMPLE_WITH_DX12
+				if (OutputSync.Sync->QueryInterface(__uuidof(IMFD3D12SynchronizationObjectCommands), (void**)DecoderSync.GetInitReference()) != S_OK)
+	#endif // ALLOW_MFSAMPLE_WITH_DX12
+				{
+					OutputSync.Sync->QueryInterface(__uuidof(ID3D12Fence), (void**)ResourceFence.GetInitReference());
+				}
+			}
+
 	#if ALLOW_MFSAMPLE_WITH_DX12
 			if (DecoderSync)
 			{
@@ -547,6 +548,12 @@ void FElectraPlayerVideoDecoderOutputPC::ShutdownPoolable()
 		lockedVideoRenderer->SampleReleasedToPool(this);
 	}
 
+	// Drop fence (if any)
+	D3DFence = nullptr;
+
+	// Drop decoder output resource
+	DecoderOutputResource = nullptr;
+
 	if (OutputType != EOutputType::HardwareWin8Plus)
 	{
 		return;
@@ -559,11 +566,6 @@ void FElectraPlayerVideoDecoderOutputPC::ShutdownPoolable()
 		Texture->QueryInterface(__uuidof(IDXGIResource), (void**)&OtherResource);
 	}
 
-	D3DFence = nullptr;
-
-	// Drop decoder output resource
-	DecoderOutputResource = nullptr;
-
 	// Make sure DX11 shared texture sync state is as expected
 	if (OtherResource)
 	{
diff --git a/Engine/Plugins/Media/ElectraPlayer/Source/ElectraPlayerRuntime/Private/Runtime/Decoder/Windows/ElectraVideoDecoder_PC.h b/Engine/Plugins/Media/ElectraPlayer/Source/ElectraPlayerRuntime/Private/Runtime/Decoder/Windows/ElectraVideoDecoder_PC.h
--- a/Engine/Plugins/Media/ElectraPlayer/Source/ElectraPlayerRuntime/Private/Runtime/Decoder/Windows/ElectraVideoDecoder_PC.h
+++ b/Engine/Plugins/Media/ElectraPlayer/Source/ElectraPlayerRuntime/Private/Runtime/Decoder/Windows/ElectraVideoDecoder_PC.h
@@ -52,7 +52,7 @@ public:
 	void InitializeWithBuffer(TSharedPtr<TArray<uint8>, ESPMode::ThreadSafe> InBuffer, uint32 InStride, FIntPoint Dim, TSharedPtr<Electra::FParamDict, ESPMode::ThreadSafe> InParamDict);
 
 	void InitializeWithResource(const TRefCountPtr<ID3D12Device>& InD3D12Device, const TRefCountPtr<ID3D12Resource> Resource, uint32 ResourcePitch, const FElectraDecoderOutputSync& OutputSync, const FIntPoint& OutputDim, TSharedPtr<Electra::FParamDict, ESPMode::ThreadSafe> InParamDict, TWeakPtr<Electra::IVideoDecoderResourceDelegate, ESPMode::ThreadSafe> ResourceDelegate,
-								uint32 MaxWidth, uint32 MaxHeight, uint32 MaxOutputBuffers);
+								TSharedPtr<FElectraMediaDecoderOutputBufferPool_DX12>& InOutD3D12ResourcePool, uint32 MaxWidth, uint32 MaxHeight, uint32 MaxOutputBuffers);
 
 	// Hardware decode to shared DX11 texture (Win8+) from IMFSample
 	void InitializeWithSharedTexture(const TRefCountPtr<ID3D11Device>& InD3D11Device, const TRefCountPtr<ID3D11DeviceContext> InDeviceContext, const TRefCountPtr<IMFSample> MFSample, const FIntPoint& OutputDim, TSharedPtr<Electra::FParamDict, ESPMode::ThreadSafe> InParamDict);
@@ -72,11 +72,7 @@ public:
 	TRefCountPtr<IUnknown> GetSync(uint64& SyncValue) const override;
 
 private:
-#if HAVE_MFSAMPLE_WITH_DX12
-	void TriggerDataCopy(TRefCountPtr<IMFD3D12SynchronizationObjectCommands> DecoderSync, TRefCountPtr<ID3D12Fence> ResourceFence, const FElectraDecoderOutputSync& OutputSync, Electra::IVideoDecoderResourceDelegate* InResourceDelegate) const;
-#else
-	void TriggerDataCopy(TRefCountPtr<IUnknown> DecoderSync, TRefCountPtr<ID3D12Fence> ResourceFence, const FElectraDecoderOutputSync& OutputSync, Electra::IVideoDecoderResourceDelegate* InResourceDelegate) const;
-#endif
+	static void TriggerDataCopy(TRefCountPtr<ID3D12GraphicsCommandList> D3DCmdList, TRefCountPtr<ID3D12Fence> D3DFence, uint64 FenceValue, const FElectraDecoderOutputSync& OutputSync, Electra::IVideoDecoderResourceDelegate* InResourceDelegate);
 
 	// Decoder output type
 	EOutputType OutputType = EOutputType::Unknown;
diff --git a/Engine/Plugins/Media/ElectraPlayer/Source/ElectraPlayerRuntime/Private/Runtime/Decoder/Windows/WindowsElectraDecoderResourceManager.cpp b/Engine/Plugins/Media/ElectraPlayer/Source/ElectraPlayerRuntime/Private/Runtime/Decoder/Windows/WindowsElectraDecoderResourceManager.cpp
--- a/Engine/Plugins/Media/ElectraPlayer/Source/ElectraPlayerRuntime/Private/Runtime/Decoder/Windows/WindowsElectraDecoderResourceManager.cpp
+++ b/Engine/Plugins/Media/ElectraPlayer/Source/ElectraPlayerRuntime/Private/Runtime/Decoder/Windows/WindowsElectraDecoderResourceManager.cpp
@@ -83,6 +83,7 @@ public:
 	uint32 MaxHeight = 0;
 	uint32 MaxOutputBuffers = 0;
 	TWeakPtr<IVideoDecoderResourceDelegate, ESPMode::ThreadSafe> VideoDecoderResourceDelegate;
+	TSharedPtr<FElectraMediaDecoderOutputBufferPool_DX12> D3D12ResourcePool;
 };
 
 
@@ -530,7 +531,8 @@ bool FElectraDecoderResourceManagerWindows::SetupRenderBufferFromDecoderOutput(I
 							FElectraDecoderOutputSync OutputSync;
 							ImageBuffers->GetBufferTextureSyncByIndex(0, OutputSync);
 
-							DecoderOutput->InitializeWithResource(D3D12Device, Resource, Pitch, OutputSync, FIntPoint(InDecoderOutput->GetDecodedWidth(), InDecoderOutput->GetDecodedHeight()), InOutBufferPropertes, Vars->VideoDecoderResourceDelegate, Vars->MaxWidth, Vars->MaxHeight, Vars->MaxOutputBuffers);
+							DecoderOutput->InitializeWithResource(D3D12Device, Resource, Pitch, OutputSync, FIntPoint(InDecoderOutput->GetDecodedWidth(), InDecoderOutput->GetDecodedHeight()), InOutBufferPropertes, Vars->VideoDecoderResourceDelegate,
+																  Vars->D3D12ResourcePool, Vars->MaxWidth, Vars->MaxHeight, Vars->MaxOutputBuffers);
 							return true;
 						}
 					}
diff --git a/Engine/Plugins/Media/ElectraPlayer/Source/ElectraPlayerRuntime/Private/Runtime/Renderer/RendererVideo.cpp b/Engine/Plugins/Media/ElectraPlayer/Source/ElectraPlayerRuntime/Private/Runtime/Renderer/RendererVideo.cpp
--- a/Engine/Plugins/Media/ElectraPlayer/Source/ElectraPlayerRuntime/Private/Runtime/Renderer/RendererVideo.cpp
+++ b/Engine/Plugins/Media/ElectraPlayer/Source/ElectraPlayerRuntime/Private/Runtime/Renderer/RendererVideo.cpp
@@ -54,8 +54,6 @@ FElectraRendererVideo::FElectraRendererVideo(TSharedPtr<FElectraPlayer, ESPMode:
 
 FElectraRendererVideo::~FElectraRendererVideo()
 {
-	// Manually delete all images in all queues. Well, the queues will are also deleted automatically.
-	QueueTickedAndWaitingForDecoder.Empty();
 }
 
 
@@ -124,8 +122,6 @@ UEMediaError FElectraRendererVideo::CreateBufferPool(const Electra::FParamDict&
 		{
 			FVideoDecoderOutputPtr DelayedImage;
 			AcquireFromPool(DelayedImage);
-			QueueTickedAndWaitingForDecoder.Enqueue(DelayedImage);
-
 		}
 		NumBuffers = RequestedNumBuffers;
 	}
@@ -150,29 +146,13 @@ UEMediaError FElectraRendererVideo::AcquireBuffer(IBuffer*& OutBuffer, int32 Tim
 		PinnedPlayer->DropOldFramesFromPresentationQueue();
 	}
 
-	// Check if we have any buffer which we can hand out to the decoder...
-	if (QueueTickedAndWaitingForDecoder.IsEmpty())
+	FVideoDecoderOutputPtr DelayedImage;
+	AcquireFromPool(DelayedImage);
+	if (!DelayedImage.IsValid())
 	{
-		// Check GLOBAL number of output textures in flight...
-		if (NumOutputTexturesInUse < NumBuffers)
-		{
-			// Allocate an image from the pool, but do NOT use it immediately because it COULD be still rendering...
-			FVideoDecoderOutputPtr DelayedImage;
-			// Acquire image and with information about how many render frames to wait before usage is possible
-			AcquireFromPool(DelayedImage);
-			QueueTickedAndWaitingForDecoder.Enqueue(DelayedImage);
-		}
-
-		if (QueueTickedAndWaitingForDecoder.IsEmpty())
-		{
-			return UEMEDIA_ERROR_INSUFFICIENT_DATA;
-		}
+		return UEMEDIA_ERROR_INSUFFICIENT_DATA;
 	}
 
-	FVideoDecoderOutputPtr DelayedImage;
-	QueueTickedAndWaitingForDecoder.Dequeue(DelayedImage);
-	check(DelayedImage.IsValid());
-
 	FMediaBufferSharedPtrWrapper* MediaBufferSharedPtrWrapper = new FMediaBufferSharedPtrWrapper(DelayedImage);
 	check(MediaBufferSharedPtrWrapper);
 	MediaBufferSharedPtrWrapper->BufferProperties.Set(RenderOptionKeys::Texture, FVariantValue(DelayedImage));
diff --git a/Engine/Plugins/Media/ElectraPlayer/Source/ElectraPlayerRuntime/Private/Runtime/Renderer/RendererVideo.h b/Engine/Plugins/Media/ElectraPlayer/Source/ElectraPlayerRuntime/Private/Runtime/Renderer/RendererVideo.h
--- a/Engine/Plugins/Media/ElectraPlayer/Source/ElectraPlayerRuntime/Private/Runtime/Renderer/RendererVideo.h
+++ b/Engine/Plugins/Media/ElectraPlayer/Source/ElectraPlayerRuntime/Private/Runtime/Renderer/RendererVideo.h
@@ -143,9 +143,6 @@ private:
 	void AcquireFromPool(FVideoDecoderOutputPtr& DelayedImage);
 
 
-	// These queues are all single-producer, single-consumer which is the default mode...
-	TQueue<FVideoDecoderOutputPtr> QueueTickedAndWaitingForDecoder;
-
 	Electra::FParamDict BufferPoolProperties;
 	int32 NumBuffers;
 	int32 NumBuffersAcquiredForDecoder;
diff --git a/Engine/Plugins/Media/ElectraPlayer/Source/ElectraPlayerRuntime/Private/Windows/ElectraPlayer_Windows.cpp b/Engine/Plugins/Media/ElectraPlayer/Source/ElectraPlayerRuntime/Private/Windows/ElectraPlayer_Windows.cpp
--- a/Engine/Plugins/Media/ElectraPlayer/Source/ElectraPlayerRuntime/Private/Windows/ElectraPlayer_Windows.cpp
+++ b/Engine/Plugins/Media/ElectraPlayer/Source/ElectraPlayerRuntime/Private/Windows/ElectraPlayer_Windows.cpp
@@ -49,23 +49,28 @@ namespace Electra
 	class FVideoDecoderResourceDelegate : public IVideoDecoderResourceDelegate
 	{
 	public:
-		FVideoDecoderResourceDelegate(const TWeakPtr<IElectraPlayerAdapterDelegate, ESPMode::ThreadSafe>& InAdapterDelegate) : PlayerAdapterDelegate(InAdapterDelegate) {}
-		virtual ~FVideoDecoderResourceDelegate() {}
+		FVideoDecoderResourceDelegate(const TWeakPtr<IElectraPlayerAdapterDelegate, ESPMode::ThreadSafe>& InAdapterDelegate)
+		: PlayerAdapterDelegate(InAdapterDelegate)
+		{
+			const TSharedPtr<IElectraPlayerAdapterDelegate, ESPMode::ThreadSafe> PinnedAdapterDelegate = PlayerAdapterDelegate.Pin();
+			check(PinnedAdapterDelegate.IsValid());
+
+			// Grab a reference to the IElectraPlayerResourceDelegate as we need it still when IElectraPlayerAdapterDelegate is already destroyed
+			ResourceDelegate = PinnedAdapterDelegate->GetResourceDelegate();
+			check(ResourceDelegate.IsValid());
+		}
+
+		virtual ~FVideoDecoderResourceDelegate() = default;
 
 		virtual void ExecuteCodeWithCopyCommandQueueUsage(TFunction<void(ID3D12CommandQueue*)>&& CodeToRun) override
 		{
-			if (const TSharedPtr<IElectraPlayerAdapterDelegate, ESPMode::ThreadSafe>& PinnedAdapterDelegate = PlayerAdapterDelegate.Pin())
-			{
-				TSharedPtr<IElectraPlayerResourceDelegate, ESPMode::ThreadSafe> ResourceDelegate = PinnedAdapterDelegate->GetResourceDelegate();
-				if (ResourceDelegate)
-				{
-					ResourceDelegate->ExecuteCodeWithCopyCommandQueueUsage(MoveTemp(CodeToRun));
-				}
-			}
+			check(ResourceDelegate.IsValid());
+			ResourceDelegate->ExecuteCodeWithCopyCommandQueueUsage(MoveTemp(CodeToRun));
 		}
 
 	private:
 		TWeakPtr<IElectraPlayerAdapterDelegate, ESPMode::ThreadSafe> PlayerAdapterDelegate;
+		TSharedPtr<IElectraPlayerResourceDelegate, ESPMode::ThreadSafe> ResourceDelegate;
 	};
 
 	TSharedPtr<IVideoDecoderResourceDelegate, ESPMode::ThreadSafe> PlatformCreateVideoDecoderResourceDelegate(const TWeakPtr<IElectraPlayerAdapterDelegate, ESPMode::ThreadSafe>& AdapterDelegate)
diff --git a/Engine/Plugins/Media/ElectraUtil/Source/ElectraBase/Public/MediaDecoderOutput.h b/Engine/Plugins/Media/ElectraUtil/Source/ElectraBase/Public/MediaDecoderOutput.h
--- a/Engine/Plugins/Media/ElectraUtil/Source/ElectraBase/Public/MediaDecoderOutput.h
+++ b/Engine/Plugins/Media/ElectraUtil/Source/ElectraBase/Public/MediaDecoderOutput.h
@@ -100,16 +100,25 @@ class TDecoderOutputObjectPool
 				}
 				else
 				{
-					if (WaitReadyForReuse.Peek(Result))
+					// Check for objects to ready to enter the pool & grab the first we find as our result...
+					// (we move all into the pool, we can to possibly safe on resources allocated by the objects)
+					ObjectType* PeekObject;
+					while (WaitReadyForReuse.Peek(PeekObject))
 					{
-						if (Result->IsReadyForReuse())
+						if (!PeekObject->IsReadyForReuse())
 						{
-							WaitReadyForReuse.Pop();
-							Result->ShutdownPoolable();
+							break;
+						}
+
+						WaitReadyForReuse.Pop();
+						PeekObject->ShutdownPoolable();
+						if (!Result)
+						{
+							Result = PeekObject;
 						}
 						else
 						{
-							Result = nullptr;
+							Pool.Add(PeekObject);
 						}
 					}
 				}
diff --git a/Engine/Source/Runtime/MediaUtils/Public/MediaObjectPool.h b/Engine/Source/Runtime/MediaUtils/Public/MediaObjectPool.h
--- a/Engine/Source/Runtime/MediaUtils/Public/MediaObjectPool.h
+++ b/Engine/Source/Runtime/MediaUtils/Public/MediaObjectPool.h
@@ -105,16 +105,25 @@ class TMediaObjectPool
 				}
 				else
 				{
-					if (WaitReadyForReuse.Peek(Result))
+					// Check for objects to ready to enter the pool & grab the first we find as our result...
+					// (we move all into the pool, we can to possibly safe on resources allocated by the objects)
+					ObjectType* PeekObject;
+					while (WaitReadyForReuse.Peek(PeekObject))
 					{
-						if (Result->IsReadyForReuse())
+						if (!PeekObject->IsReadyForReuse())
 						{
-							WaitReadyForReuse.Pop();
-							Result->ShutdownPoolable();
+							break;
+						}
+
+						WaitReadyForReuse.Pop();
+						PeekObject->ShutdownPoolable();
+						if (!Result)
+						{
+							Result = PeekObject;
 						}
 						else
 						{
-							Result = nullptr;
+							Pool.Add(PeekObject);
 						}
 					}
 				}
