Optimizing flags used to create media playback heaps in Electra for speed & compatibility
#rnx
#tests Tested on local PC with a number of videos

#changelist validated
#virtualized

diff --git a/Engine/Plugins/Media/ElectraCodecs/Source/ElectraDecoders/Public/Windows/WindowsElectraDecoderGPUBufferHelpers.h b/Engine/Plugins/Media/ElectraCodecs/Source/ElectraDecoders/Public/Windows/WindowsElectraDecoderGPUBufferHelpers.h
--- a/Engine/Plugins/Media/ElectraCodecs/Source/ElectraDecoders/Public/Windows/WindowsElectraDecoderGPUBufferHelpers.h
+++ b/Engine/Plugins/Media/ElectraCodecs/Source/ElectraDecoders/Public/Windows/WindowsElectraDecoderGPUBufferHelpers.h
@@ -191,10 +191,10 @@ private:
 	void InitCommon()
 	{
 		D3D12_HEAP_DESC HeapDesc = {};
-		HeapDesc.Alignment = D3D12_DEFAULT_RESOURCE_PLACEMENT_ALIGNMENT;
+		HeapDesc.Alignment = 0;
 		HeapDesc.SizeInBytes = BufferSize * MaxNumBuffers;
 		HeapDesc.Properties = CD3DX12_HEAP_PROPERTIES(D3D12HeapType);
-		HeapDesc.Flags = (D3D12HeapType == D3D12_HEAP_TYPE_UPLOAD) ? D3D12_HEAP_FLAG_ALLOW_ONLY_BUFFERS : D3D12_HEAP_FLAG_ALLOW_ALL_BUFFERS_AND_TEXTURES;
+		HeapDesc.Flags = D3D12_HEAP_FLAG_CREATE_NOT_ZEROED | ((D3D12HeapType == D3D12_HEAP_TYPE_UPLOAD) ? D3D12_HEAP_FLAG_ALLOW_ONLY_BUFFERS : D3D12_HEAP_FLAG_ALLOW_ONLY_NON_RT_DS_TEXTURES);
 		HRESULT Res = D3D12Device->CreateHeap(&HeapDesc, IID_PPV_ARGS(D3D12OutputHeap.GetInitReference()));
 		check(SUCCEEDED(Res));
 		if (Res != S_OK)
@@ -202,7 +202,10 @@ private:
 			FreeMask = 0;
 		}
 #if !UE_BUILD_SHIPPING
-		D3D12OutputHeap->SetName(TEXT("ElectraOutputBufferPoolHeap"));
+		if (D3D12OutputHeap)
+		{
+			D3D12OutputHeap->SetName(TEXT("ElectraOutputBufferPoolHeap"));
+		}
 #endif
 	}
 
