Allowing for DX12 buffer reinitialization in Electra output pipeline
#tests Tested in tw separate setups on local machines

diff --git a/Engine/Plugins/Media/ElectraCodecs/Source/ElectraDecoders/Private/DX/VideoDecoderH264_DX.cpp b/Engine/Plugins/Media/ElectraCodecs/Source/ElectraDecoders/Private/DX/VideoDecoderH264_DX.cpp
--- a/Engine/Plugins/Media/ElectraCodecs/Source/ElectraDecoders/Private/DX/VideoDecoderH264_DX.cpp
+++ b/Engine/Plugins/Media/ElectraCodecs/Source/ElectraDecoders/Private/DX/VideoDecoderH264_DX.cpp
@@ -1148,7 +1148,7 @@ bool FElectraVideoDecoderH264_DX::ConvertDecoderOutput()
 		TRefCountPtr D3D12Device(static_cast<ID3D12Device*>(PlatformDevice));
 
 		// Create the resource pool as needed...
-		if (!D3D12ResourcePool)
+		if (!D3D12ResourcePool || !D3D12ResourcePool->IsCompatibleAsBuffer(MaxOutputBuffers, MaxWidth, MaxHeight * 3 / 2, 1))
 		{
 			D3D12ResourcePool = MakeShared<FElectraMediaDecoderOutputBufferPool_DX12, ESPMode::ThreadSafe>(D3D12Device, MaxOutputBuffers, MaxWidth, MaxHeight * 3 / 2, 1);
 		}
diff --git a/Engine/Plugins/Media/ElectraCodecs/Source/ElectraDecoders/Private/DX/VideoDecoderH265_DX.cpp b/Engine/Plugins/Media/ElectraCodecs/Source/ElectraDecoders/Private/DX/VideoDecoderH265_DX.cpp
--- a/Engine/Plugins/Media/ElectraCodecs/Source/ElectraDecoders/Private/DX/VideoDecoderH265_DX.cpp
+++ b/Engine/Plugins/Media/ElectraCodecs/Source/ElectraDecoders/Private/DX/VideoDecoderH265_DX.cpp
@@ -1206,7 +1206,7 @@ bool FElectraVideoDecoderH265_DX::ConvertDecoderOutput()
 		TRefCountPtr D3D12Device(static_cast<ID3D12Device*>(PlatformDevice));
 
 		// Create the resource pool as needed...
-		if (!D3D12ResourcePool)
+		if (!D3D12ResourcePool || !D3D12ResourcePool->IsCompatibleAsBuffer(MaxOutputBuffers, MaxWidth, MaxHeight * 3 / 2, (PixFmt == EElectraDecoderPlatformPixelFormat::NV12) ? 1 : 2))
 		{
 			D3D12ResourcePool = MakeShared<FElectraMediaDecoderOutputBufferPool_DX12, ESPMode::ThreadSafe>(D3D12Device, MaxOutputBuffers, MaxWidth, MaxHeight * 3 / 2, (PixFmt == EElectraDecoderPlatformPixelFormat::NV12) ? 1 : 2);
 		}
diff --git a/Engine/Plugins/Media/ElectraCodecs/Source/ElectraDecoders/Public/Windows/WindowsElectraDecoderGPUBufferHelpers.h b/Engine/Plugins/Media/ElectraCodecs/Source/ElectraDecoders/Public/Windows/WindowsElectraDecoderGPUBufferHelpers.h
--- a/Engine/Plugins/Media/ElectraCodecs/Source/ElectraDecoders/Public/Windows/WindowsElectraDecoderGPUBufferHelpers.h
+++ b/Engine/Plugins/Media/ElectraCodecs/Source/ElectraDecoders/Public/Windows/WindowsElectraDecoderGPUBufferHelpers.h
@@ -139,7 +139,29 @@ public:
 	{
 		uint32 NewBufferPitch = Align(Width * BytesPerPixel, D3D12_TEXTURE_DATA_PITCH_ALIGNMENT);
 		uint32 NewBufferSize = Align(NewBufferPitch * Height, FMath::Max(D3D12_TEXTURE_DATA_PLACEMENT_ALIGNMENT, D3D12_DEFAULT_RESOURCE_PLACEMENT_ALIGNMENT));
-		return (MaxNumBuffers == InMaxNumBuffers && BufferSize == NewBufferSize && BufferPitch == NewBufferPitch);
+		return (MaxNumBuffers >= InMaxNumBuffers && BufferSize >= NewBufferSize && BufferPitch == NewBufferPitch);
+	}
+
+	// Check if the current setup is compatible with the new parameters
+	bool IsCompatibleAsTexture(uint32 InMaxNumBuffers, uint32 Width, uint32 Height, DXGI_FORMAT PixFmt) const
+	{
+		// Get the size needed per buffer / texture
+		D3D12_RESOURCE_DESC Desc = {};
+		Desc.MipLevels = 1;
+		Desc.Format = PixFmt;
+		Desc.Width = Width;
+		Desc.Height = Height;
+		Desc.Flags = D3D12_RESOURCE_FLAG_NONE;
+		Desc.DepthOrArraySize = 1;
+		Desc.SampleDesc.Count = 1;
+		Desc.SampleDesc.Quality = 0;
+		Desc.Layout = D3D12_TEXTURE_LAYOUT_UNKNOWN;
+		Desc.Dimension = D3D12_RESOURCE_DIMENSION_TEXTURE2D;
+		D3D12_RESOURCE_ALLOCATION_INFO AllocInfo = D3D12Device->GetResourceAllocationInfo(0, 1, &Desc);
+
+		uint32 NewBufferSize = Align(AllocInfo.SizeInBytes, AllocInfo.Alignment);
+
+		return (MaxNumBuffers >= InMaxNumBuffers && BufferSize >= NewBufferSize);
 	}
 
 	// Check if a buffer is available
@@ -337,6 +359,11 @@ public:
 		return Blocks[0]->IsCompatibleAsBuffer(InMaxNumBuffers, InWidth, InHeight, InBytesPerPixel);
 	}
 
+	bool IsCompatibleAsTexture(uint32 InMaxNumBuffers, uint32 InWidth, uint32 InHeight, DXGI_FORMAT InPixFmt) const
+	{
+		return Blocks[0]->IsCompatibleAsTexture(InMaxNumBuffers, InWidth, InHeight, InPixFmt);
+	}
+
 	// Check if a buffer is available
 	bool BufferAvailable() const
 	{
diff --git a/Engine/Plugins/Media/ElectraPlayer/Source/ElectraPlayerRuntime/Private/Runtime/Decoder/Windows/ElectraVideoDecoder_PC.cpp b/Engine/Plugins/Media/ElectraPlayer/Source/ElectraPlayerRuntime/Private/Runtime/Decoder/Windows/ElectraVideoDecoder_PC.cpp
--- a/Engine/Plugins/Media/ElectraPlayer/Source/ElectraPlayerRuntime/Private/Runtime/Decoder/Windows/ElectraVideoDecoder_PC.cpp
+++ b/Engine/Plugins/Media/ElectraPlayer/Source/ElectraPlayerRuntime/Private/Runtime/Decoder/Windows/ElectraVideoDecoder_PC.cpp
@@ -221,7 +221,7 @@ void FElectraPlayerVideoDecoderOutputPC::InitializeWithResource(const TRefCountP
 		// We get here only after the instance came back from the pool! Hence we can be sure any old texture is no longer actively used.
 		TextureDX12 = nullptr;
 
-		if (!D3D12ResourcePool.IsValid())
+		if (!D3D12ResourcePool.IsValid() || !D3D12ResourcePool->IsCompatibleAsTexture(MaxOutputBuffers + kElectraDecoderPipelineExtraFrames, MaxWidth, MaxHeight, DXGIFmt))
 		{
 			D3D12ResourcePool = MakeShared<FElectraMediaDecoderOutputBufferPool_DX12>(InD3D12Device, MaxOutputBuffers + kElectraDecoderPipelineExtraFrames, MaxWidth, MaxHeight, DXGIFmt, D3D12_HEAP_TYPE_DEFAULT);
 			if (!D3D12ResourcePool.IsValid())
