ElectraDecoders: allowing for 120fps on Apple devices

#rnx

diff --git a/Engine/Plugins/Media/ElectraCodecs/Source/ElectraDecoders/Private/Apple/h264/VideoDecoderH264_Apple.cpp b/Engine/Plugins/Media/ElectraCodecs/Source/ElectraDecoders/Private/Apple/h264/VideoDecoderH264_Apple.cpp
--- a/Engine/Plugins/Media/ElectraCodecs/Source/ElectraDecoders/Private/Apple/h264/VideoDecoderH264_Apple.cpp
+++ b/Engine/Plugins/Media/ElectraCodecs/Source/ElectraDecoders/Private/Apple/h264/VideoDecoderH264_Apple.cpp
@@ -364,11 +364,11 @@ void IElectraVideoDecoderH264_Apple::PlatformGetSupportedConfigurations(TArray<F
 		// What is really supported isn't quite clear so allow for UHD.
 
 		// Baseline
-		IElectraVideoDecoderH264_Apple_Platform::DecoderConfigurations.Emplace(IElectraVideoDecoderH264_Apple::FSupportedConfiguration(66, 52, 60, 3840, 2160, 0));
+		IElectraVideoDecoderH264_Apple_Platform::DecoderConfigurations.Emplace(IElectraVideoDecoderH264_Apple::FSupportedConfiguration(66, 52, 120, 3840, 2160, 0));
 		// Main
-		IElectraVideoDecoderH264_Apple_Platform::DecoderConfigurations.Emplace(IElectraVideoDecoderH264_Apple::FSupportedConfiguration(77, 52, 60, 3840, 2160, 0));
+		IElectraVideoDecoderH264_Apple_Platform::DecoderConfigurations.Emplace(IElectraVideoDecoderH264_Apple::FSupportedConfiguration(77, 52, 120, 3840, 2160, 0));
 		// High
-		IElectraVideoDecoderH264_Apple_Platform::DecoderConfigurations.Emplace(IElectraVideoDecoderH264_Apple::FSupportedConfiguration(100, 52, 60, 3840, 2160, 0));
+		IElectraVideoDecoderH264_Apple_Platform::DecoderConfigurations.Emplace(IElectraVideoDecoderH264_Apple::FSupportedConfiguration(100, 52, 120, 3840, 2160, 0));
 
 		IElectraVideoDecoderH264_Apple_Platform::bDecoderConfigurationsDirty = false;
 	}
diff --git a/Engine/Plugins/Media/ElectraCodecs/Source/ElectraDecoders/Private/Apple/h265/VideoDecoderH265_Apple.cpp b/Engine/Plugins/Media/ElectraCodecs/Source/ElectraDecoders/Private/Apple/h265/VideoDecoderH265_Apple.cpp
--- a/Engine/Plugins/Media/ElectraCodecs/Source/ElectraDecoders/Private/Apple/h265/VideoDecoderH265_Apple.cpp
+++ b/Engine/Plugins/Media/ElectraCodecs/Source/ElectraDecoders/Private/Apple/h265/VideoDecoderH265_Apple.cpp
@@ -364,11 +364,11 @@ void IElectraVideoDecoderH265_Apple::PlatformGetSupportedConfigurations(TArray<F
 		// What is really supported isn't quite clear so allow for UHD.
 
 		// Main
-		IElectraVideoDecoderH265_Apple_Platform::DecoderConfigurations.Emplace(IElectraVideoDecoderH265_Apple::FSupportedConfiguration(0, 1, 0, 153, 60, 4096, 2304, 0));
-		IElectraVideoDecoderH265_Apple_Platform::DecoderConfigurations.Emplace(IElectraVideoDecoderH265_Apple::FSupportedConfiguration(0, 1, 0, 153, 0, 0, 0, (3840 / 8) * (2160 / 8) * 60));
+		IElectraVideoDecoderH265_Apple_Platform::DecoderConfigurations.Emplace(IElectraVideoDecoderH265_Apple::FSupportedConfiguration(0, 1, 0, 153, 120, 4096, 2304, 0));
+		IElectraVideoDecoderH265_Apple_Platform::DecoderConfigurations.Emplace(IElectraVideoDecoderH265_Apple::FSupportedConfiguration(0, 1, 0, 153, 0, 0, 0, (3840 / 8) * (2160 / 8) * 120));
 		// Main10
-		IElectraVideoDecoderH265_Apple_Platform::DecoderConfigurations.Emplace(IElectraVideoDecoderH265_Apple::FSupportedConfiguration(0, 2, 0, 153, 60, 4096, 2304, 0));
-		IElectraVideoDecoderH265_Apple_Platform::DecoderConfigurations.Emplace(IElectraVideoDecoderH265_Apple::FSupportedConfiguration(0, 2, 0, 153, 0, 0, 0, (3840 / 8) * (2160 / 8) * 60));
+		IElectraVideoDecoderH265_Apple_Platform::DecoderConfigurations.Emplace(IElectraVideoDecoderH265_Apple::FSupportedConfiguration(0, 2, 0, 153, 120, 4096, 2304, 0));
+		IElectraVideoDecoderH265_Apple_Platform::DecoderConfigurations.Emplace(IElectraVideoDecoderH265_Apple::FSupportedConfiguration(0, 2, 0, 153, 0, 0, 0, (3840 / 8) * (2160 / 8) * 120));
 
 		IElectraVideoDecoderH265_Apple_Platform::bDecoderConfigurationsDirty = false;
 	}
