ElectraDecoders: fixed a code analysis warning

#rnx

diff --git a/Engine/Plugins/Media/ElectraCodecs/Source/ElectraDecoders/Private/Utils/MPEG/ElectraUtilsMPEGVideo_H264.cpp b/Engine/Plugins/Media/ElectraCodecs/Source/ElectraDecoders/Private/Utils/MPEG/ElectraUtilsMPEGVideo_H264.cpp
--- a/Engine/Plugins/Media/ElectraCodecs/Source/ElectraDecoders/Private/Utils/MPEG/ElectraUtilsMPEGVideo_H264.cpp
+++ b/Engine/Plugins/Media/ElectraCodecs/Source/ElectraDecoders/Private/Utils/MPEG/ElectraUtilsMPEGVideo_H264.cpp
@@ -1181,7 +1181,7 @@ namespace ElectraDecodersUtil
 				CurrentPOC = PreviousPOC;
 			}
 
-			bool FSlicePOCVars::HandleMissingFrames(TArray<FOutputFrameInfo>& OutOutputFrameInfos, TArray<FOutputFrameInfo>& OutUnrefFrameInfos, uint8 InNalUnitType, uint8 InNalRefIdc, const FSliceHeader& InSliceHeader, const FSequenceParameterSet& InSequenceParameterSet)
+			bool FSlicePOCVars::HandleMissingFrames(TArray<FOutputFrameInfo>& OutOutputFrameInfos, TArray<FOutputFrameInfo>& OutUnrefFrameInfos, uint8 InNalUnitType, uint8 /*InNalRefIdc*/, const FSliceHeader& InSliceHeader, const FSequenceParameterSet& InSequenceParameterSet)
 			{
 				bool bIsIDR = InNalUnitType == 5;
 				// Do we need to check?
@@ -1207,7 +1207,7 @@ namespace ElectraDecodersUtil
 						delta_pic_order_cnt[ 0 ] (if needed) inferred to be equal to 0, and delta_pic_order_cnt[ 1 ] (if needed) inferred to be equal to 0.
 				*/
 				InNalUnitType = 1;
-				InNalRefIdc = 1;
+				uint8 InNalRefIdc = 1;
 				FSliceHeader TempSliceHeader(InSliceHeader);
 				TempSliceHeader.delta_pic_order_cnt[0] = TempSliceHeader.delta_pic_order_cnt[1] = 0;
 				TempSliceHeader.adaptive_ref_pic_marking_mode_flag = 0;
