ElectraCodecs: Added MPEG audio decoder

#rnx

diff --git a/Engine/Plugins/Media/ElectraCodecs/Source/ElectraDecoders/ElectraDecoders.Build.cs b/Engine/Plugins/Media/ElectraCodecs/Source/ElectraDecoders/ElectraDecoders.Build.cs
--- a/Engine/Plugins/Media/ElectraCodecs/Source/ElectraDecoders/ElectraDecoders.Build.cs
+++ b/Engine/Plugins/Media/ElectraCodecs/Source/ElectraDecoders/ElectraDecoders.Build.cs
@@ -27,7 +27,9 @@ namespace UnrealBuildTool.Rules
 			PrivateDependencyModuleNames.AddRange(
 				new string[] {
 					"Core",
-					"SignalProcessing"
+					"SignalProcessing",
+                    // Common codecs
+					"dr_libs"					// MPEG audio layer 3
 				});
 
 			PublicIncludePathModuleNames.AddRange(
diff --git a/Engine/Plugins/Media/ElectraCodecs/Source/ElectraDecoders/Private/ElectraDecodersModule.cpp b/Engine/Plugins/Media/ElectraCodecs/Source/ElectraDecoders/Private/ElectraDecodersModule.cpp
--- a/Engine/Plugins/Media/ElectraCodecs/Source/ElectraDecoders/Private/ElectraDecodersModule.cpp
+++ b/Engine/Plugins/Media/ElectraCodecs/Source/ElectraDecoders/Private/ElectraDecodersModule.cpp
@@ -6,9 +6,11 @@
 #include "Templates/SharedPointer.h"
 
 #include "IElectraDecodersModule.h"
-
 #include "PlatformElectraDecoders.h"
 
+// Common codecs
+#include "ElectraMediaMP3Decoder.h"
+
 #define LOCTEXT_NAMESPACE "ElectraDecodersModule"
 
 DEFINE_LOG_CATEGORY(LogElectraDecoders);
@@ -20,13 +22,15 @@ public:
 	void StartupModule() override
 	{
 		FPlatformElectraDecoders::Startup();
+		FElectraMediaMP3Decoder::Startup();
 	}
 
 	void ShutdownModule() override
 	{
+		FElectraMediaMP3Decoder::Shutdown();
 		FPlatformElectraDecoders::Shutdown();
 	}
-	
+
 	bool SupportsDynamicReloading() override
 	{
 		return false;
diff --git a/Engine/Plugins/Media/ElectraCodecs/Source/ElectraDecoders/Private/ElectraMediaMP3Decoder.cpp b/Engine/Plugins/Media/ElectraCodecs/Source/ElectraDecoders/Private/ElectraMediaMP3Decoder.cpp
new file mode 100644
--- /dev/null
+++ b/Engine/Plugins/Media/ElectraCodecs/Source/ElectraDecoders/Private/ElectraMediaMP3Decoder.cpp
@@ -0,0 +1,576 @@
+// Copyright Epic Games, Inc. All Rights Reserved.
+
+#include "ElectraMediaMP3Decoder.h"
+#include "IElectraCodecRegistry.h"
+#include "IElectraCodecFactory.h"
+#include "Misc/ScopeLock.h"
+#include "Templates/SharedPointer.h"
+#include "Features/IModularFeatures.h"
+#include "Features/IModularFeature.h"
+#include "IElectraCodecFactoryModule.h"
+#include "IElectraDecoderFeaturesAndOptions.h"
+#include "IElectraDecoder.h"
+#include "IElectraDecoderOutputAudio.h"
+#include "IElectraDecoderResourceDelegate.h"
+#include "ElectraDecodersUtils.h"
+#include "Utils/ElectraBitstreamReader.h"
+#include "Utils/MPEG/ElectraUtilsMPEGAudio.h"
+
+#ifndef WITH_DR_LIBS_MP3
+#define WITH_DR_LIBS_MP3 0
+#endif
+
+#if WITH_DR_LIBS_MP3
+#include "dr_libs_mp3decoder.h"
+#endif
+
+#define ERRCODE_MP3DEC_INTERNAL_NO_ERROR						0
+#define ERRCODE_MP3DEC_INTERNAL_ALREADY_CLOSED					1
+#define ERRCODE_MP3DEC_INTERNAL_FAILED_TO_PARSE_CSD				2
+#define ERRCODE_MP3DEC_INTERNAL_FAILED_TO_DECODE_INPUT			3
+#define ERRCODE_MP3DEC_INTERNAL_UNSUPPORTED_CHANNEL_LAYOUT		4
+
+/*********************************************************************************************************************/
+/*********************************************************************************************************************/
+/*********************************************************************************************************************/
+
+class FElectraDecoderDefaultAudioOutputFormatMP3_Common : public IElectraDecoderDefaultAudioOutputFormat
+{
+public:
+	virtual ~FElectraDecoderDefaultAudioOutputFormatMP3_Common()
+	{ }
+
+	int32 GetNumChannels() const override
+	{
+		return NumChannels;
+	}
+	int32 GetSampleRate() const override
+	{
+		return SampleRate;
+	}
+	int32 GetNumFrames() const override
+	{
+		return NumFrames;
+	}
+
+	int32 NumChannels = 0;
+	int32 SampleRate = 0;
+	int32 NumFrames = 0;
+};
+
+
+class FElectraAudioDecoderOutputMP3_Common : public IElectraDecoderAudioOutput
+{
+public:
+	virtual ~FElectraAudioDecoderOutputMP3_Common()
+	{
+		FMemory::Free(Buffer);
+	}
+
+	FTimespan GetPTS() const override
+	{
+		return PTS;
+	}
+	uint64 GetUserValue() const override
+	{
+		return UserValue;
+	}
+
+	int32 GetNumChannels() const override
+	{
+		return NumChannels;
+	}
+	int32 GetSampleRate() const override
+	{
+		return SampleRate;
+	}
+	int32 GetNumFrames() const override
+	{
+		return NumFrames;
+	}
+	bool IsInterleaved() const override
+	{
+		return true;
+	}
+	EChannelPosition GetChannelPosition(int32 InChannelNumber) const override
+	{
+		return InChannelNumber >= 0 && InChannelNumber < ChannelPositions.Num() ? ChannelPositions[InChannelNumber] : EChannelPosition::Invalid;
+	}
+	ESampleFormat GetSampleFormat() const override
+	{
+		return ESampleFormat::Float;
+	}
+	int32 GetBytesPerSample() const override
+	{
+		return sizeof(float);
+	}
+	int32 GetBytesPerFrame() const override
+	{
+		return GetBytesPerSample() * GetNumChannels();
+	}
+	const void* GetData(int32 InChannelNumber) const override
+	{
+		return InChannelNumber >= 0 && InChannelNumber < GetNumChannels() ? Buffer + InChannelNumber : nullptr;
+	}
+
+public:
+	TArray<EChannelPosition> ChannelPositions;
+	FTimespan PTS;
+	float* Buffer = nullptr;
+	uint64 UserValue = 0;
+	int32 NumChannels = 0;
+	int32 SampleRate = 0;
+	int32 NumFrames = 0;
+};
+
+
+class FElectraMP3Decoder : public IElectraDecoder
+{
+public:
+	static void GetConfigurationOptions(TMap<FString, FVariant>& OutOptions)
+	{
+	}
+
+	FElectraMP3Decoder(const TMap<FString, FVariant>& InOptions, TSharedPtr<IElectraDecoderResourceDelegate, ESPMode::ThreadSafe> InResourceDelegate);
+
+	virtual ~FElectraMP3Decoder();
+
+	IElectraDecoder::EType GetType() const override
+	{
+		return IElectraDecoder::EType::Audio;
+	}
+
+	void GetFeatures(TMap<FString, FVariant>& OutFeatures) const override;
+
+	FError GetError() const override;
+
+	void Close() override;
+	IElectraDecoder::ECSDCompatibility IsCompatibleWith(const TMap<FString, FVariant>& CSDAndAdditionalOptions) override;
+	bool ResetToCleanStart() override;
+
+	TSharedPtr<IElectraDecoderDefaultOutputFormat, ESPMode::ThreadSafe> GetDefaultOutputFormatFromCSD(const TMap<FString, FVariant>& CSDAndAdditionalOptions) override;
+
+	EDecoderError DecodeAccessUnit(const FInputAccessUnit& InInputAccessUnit, const TMap<FString, FVariant>& InAdditionalOptions) override;
+	EDecoderError SendEndOfData() override;
+	EDecoderError Flush() override;
+	EOutputStatus HaveOutput() override;
+	TSharedPtr<IElectraDecoderOutput, ESPMode::ThreadSafe> GetOutput() override;
+
+	void Suspend() override
+	{ }
+	void Resume() override
+	{ }
+
+private:
+	struct FMpegAudioConfig
+	{
+		int32 SampleRate = 0;
+		int32 NumberOfChannels = 0;
+		int32 SamplesPerFrame = 0;
+		int32 Bitrate = 0;
+		void Reset()
+		{
+			SampleRate = 0;
+			NumberOfChannels = 0;
+		}
+		bool SameAs(const FMpegAudioConfig& rhs)
+		{
+			return SampleRate == rhs.SampleRate && NumberOfChannels == rhs.NumberOfChannels;
+		}
+	};
+
+	static constexpr uint32 Make4CC(const uint8 A, const uint8 B, const uint8 C, const uint8 D)
+	{
+		return (static_cast<uint32>(A) << 24) | (static_cast<uint32>(B) << 16) | (static_cast<uint32>(C) << 8) | static_cast<uint32>(D);
+	}
+
+	bool ParseMpegAudioHeader(FMpegAudioConfig& OutConfig, const void* InData, int64 InDataSize);
+
+	bool PostError(int32 ApiReturnValue, FString Message, int32 Code);
+
+	bool InternalDecoderCreate();
+	void InternalDecoderDestroy();
+
+	bool SetupChannelMap();
+
+	IElectraDecoder::FError LastError;
+
+#if WITH_DR_LIBS_MP3
+	FMP3Decoder* DecoderHandle = nullptr;
+#else
+	void* DecoderHandle = nullptr;
+#endif
+
+	uint32 Codec4CC = 0;
+	TSharedPtr<FElectraAudioDecoderOutputMP3_Common, ESPMode::ThreadSafe> CurrentOutput;
+	bool bFlushPending = false;
+
+	// Input configuration
+	FMpegAudioConfig MpegConfig;
+	bool bHaveParsedMpegHeader = false;
+
+	// Output
+	TArray<IElectraDecoderAudioOutput::EChannelPosition> OutputChannelMap;
+};
+
+
+/*********************************************************************************************************************/
+/*********************************************************************************************************************/
+/*********************************************************************************************************************/
+
+class FElectraCommonAudioFlacDecoderFactory : public IElectraCodecFactory, public IElectraCodecModularFeature, public TSharedFromThis<FElectraCommonAudioFlacDecoderFactory, ESPMode::ThreadSafe>
+{
+public:
+	virtual ~FElectraCommonAudioFlacDecoderFactory()
+	{ }
+
+	void GetListOfFactories(TArray<TWeakPtr<IElectraCodecFactory, ESPMode::ThreadSafe>>& OutCodecFactories) override
+	{
+		OutCodecFactories.Add(AsShared());
+	}
+
+	int32 SupportsFormat(const FString& InCodecFormat, bool bInEncoder, const TMap<FString, FVariant>& InOptions) const override
+	{
+		// Quick check if this is an ask for an encoder or for a 4CC we do not support.
+		if (bInEncoder || !Permitted4CCs.Contains(InCodecFormat))
+		{
+			return 0;
+		}
+		return 5;
+	}
+
+	void GetConfigurationOptions(TMap<FString, FVariant>& OutOptions) const override
+	{
+		FElectraMP3Decoder::GetConfigurationOptions(OutOptions);
+	}
+
+	TSharedPtr<IElectraDecoder, ESPMode::ThreadSafe> CreateDecoderForFormat(const FString& InCodecFormat, const TMap<FString, FVariant>& InOptions, TSharedPtr<IElectraDecoderResourceDelegate, ESPMode::ThreadSafe> InResourceDelegate) override
+	{
+		return MakeShared<FElectraMP3Decoder, ESPMode::ThreadSafe>(InOptions, InResourceDelegate);
+	}
+
+	static TSharedPtr<FElectraCommonAudioFlacDecoderFactory, ESPMode::ThreadSafe> Self;
+	static TArray<FString> Permitted4CCs;
+};
+TSharedPtr<FElectraCommonAudioFlacDecoderFactory, ESPMode::ThreadSafe> FElectraCommonAudioFlacDecoderFactory::Self;
+TArray<FString> FElectraCommonAudioFlacDecoderFactory::Permitted4CCs = { TEXT("mp4a.6b"), TEXT(".mp3") };
+
+/*********************************************************************************************************************/
+/*********************************************************************************************************************/
+/*********************************************************************************************************************/
+
+void FElectraMediaMP3Decoder::Startup()
+{
+#if WITH_DR_LIBS_MP3
+	// Make sure the codec factory module has been loaded.
+	FModuleManager::Get().LoadModule(TEXT("ElectraCodecFactory"));
+
+	// Create an instance of the factory, which is also the modular feature.
+	check(!FElectraCommonAudioFlacDecoderFactory::Self.IsValid());
+	FElectraCommonAudioFlacDecoderFactory::Self = MakeShared<FElectraCommonAudioFlacDecoderFactory, ESPMode::ThreadSafe>();
+	// Register as modular feature.
+	IModularFeatures::Get().RegisterModularFeature(IElectraCodecFactoryModule::GetModularFeatureName(), FElectraCommonAudioFlacDecoderFactory::Self.Get());
+#endif
+}
+
+void FElectraMediaMP3Decoder::Shutdown()
+{
+#if WITH_DR_LIBS_MP3
+	IModularFeatures::Get().UnregisterModularFeature(IElectraCodecFactoryModule::GetModularFeatureName(), FElectraCommonAudioFlacDecoderFactory::Self.Get());
+	FElectraCommonAudioFlacDecoderFactory::Self.Reset();
+#endif
+}
+
+TSharedPtr<IElectraCodecFactory, ESPMode::ThreadSafe> FElectraMediaMP3Decoder::CreateFactory()
+{
+	return MakeShared<FElectraCommonAudioFlacDecoderFactory, ESPMode::ThreadSafe>();;
+}
+
+/*********************************************************************************************************************/
+/*********************************************************************************************************************/
+/*********************************************************************************************************************/
+
+FElectraMP3Decoder::FElectraMP3Decoder(const TMap<FString, FVariant>& InOptions, TSharedPtr<IElectraDecoderResourceDelegate, ESPMode::ThreadSafe> InResourceDelegate)
+{
+	Codec4CC = (uint32)ElectraDecodersUtil::GetVariantValueSafeU64(InOptions, TEXT("codec_4cc"), 0);
+}
+
+FElectraMP3Decoder::~FElectraMP3Decoder()
+{
+	// Close() must have been called already!
+	check(LastError.Code == ERRCODE_MP3DEC_INTERNAL_ALREADY_CLOSED);
+	// We do it nonetheless...
+	Close();
+}
+
+bool FElectraMP3Decoder::ParseMpegAudioHeader(FMpegAudioConfig& OutConfig, const void* InData, int64 InDataSize)
+{
+	if (InDataSize >= 4)
+	{
+		const uint8* Data = reinterpret_cast<const uint8*>(InData);
+		const uint32 HeaderValue = (static_cast<uint32>(Data[0]) << 24) | (static_cast<uint32>(Data[1]) << 16) | (static_cast<uint32>(Data[2]) << 8) | static_cast<uint32>(Data[3]);
+		if (ElectraDecodersUtil::MPEG::UtilsMPEG123::HasValidSync(HeaderValue))
+		{
+			OutConfig.SampleRate = ElectraDecodersUtil::MPEG::UtilsMPEG123::GetSamplingRate(HeaderValue);
+			OutConfig.NumberOfChannels = ElectraDecodersUtil::MPEG::UtilsMPEG123::GetChannelCount(HeaderValue);
+			OutConfig.Bitrate = ElectraDecodersUtil::MPEG::UtilsMPEG123::GetChannelCount(HeaderValue);
+			OutConfig.SamplesPerFrame = ElectraDecodersUtil::MPEG::UtilsMPEG123::GetSamplesPerFrame(HeaderValue);
+			return true;
+		}
+	}
+	return false;
+}
+
+
+
+bool FElectraMP3Decoder::InternalDecoderCreate()
+{
+#if WITH_DR_LIBS_MP3
+	if (!DecoderHandle)
+	{
+		DecoderHandle = new FMP3Decoder;
+	}
+	return true;
+#else
+	return false;
+#endif
+}
+
+void FElectraMP3Decoder::InternalDecoderDestroy()
+{
+#if WITH_DR_LIBS_MP3
+	if (DecoderHandle)
+	{
+		delete DecoderHandle;
+		DecoderHandle = nullptr;
+	}
+#endif
+}
+
+
+void FElectraMP3Decoder::GetFeatures(TMap<FString, FVariant>& OutFeatures) const
+{
+	GetConfigurationOptions(OutFeatures);
+}
+
+IElectraDecoder::FError FElectraMP3Decoder::GetError() const
+{
+	return LastError;
+}
+
+bool FElectraMP3Decoder::PostError(int32 ApiReturnValue, FString Message, int32 Code)
+{
+	LastError.Code = Code;
+	LastError.SdkCode = ApiReturnValue;
+	LastError.Message = MoveTemp(Message);
+	return false;
+}
+
+void FElectraMP3Decoder::Close()
+{
+	ResetToCleanStart();
+	// Set the error state that all subsequent calls will fail.
+	PostError(0, TEXT("Already closed"), ERRCODE_MP3DEC_INTERNAL_ALREADY_CLOSED);
+}
+
+IElectraDecoder::ECSDCompatibility FElectraMP3Decoder::IsCompatibleWith(const TMap<FString, FVariant>& CSDAndAdditionalOptions)
+{
+	// No configuration parsed yet, so this is deemed compatible.
+	if (!bHaveParsedMpegHeader)
+	{
+		return IElectraDecoder::ECSDCompatibility::Compatible;
+	}
+	// There is no CSD for MPEG audio, so we can't assume compatibility
+	return IElectraDecoder::ECSDCompatibility::DrainAndReset;
+}
+
+bool FElectraMP3Decoder::ResetToCleanStart()
+{
+	bFlushPending = false;
+	CurrentOutput.Reset();
+
+	bHaveParsedMpegHeader = false;
+	MpegConfig.Reset();
+	OutputChannelMap.Empty();
+	InternalDecoderDestroy();
+	return true;
+}
+
+TSharedPtr<IElectraDecoderDefaultOutputFormat, ESPMode::ThreadSafe> FElectraMP3Decoder::GetDefaultOutputFormatFromCSD(const TMap<FString, FVariant>& CSDAndAdditionalOptions)
+{
+	return nullptr;
+}
+
+bool FElectraMP3Decoder::SetupChannelMap()
+{
+	if (OutputChannelMap.Num())
+	{
+		return true;
+	}
+	// Pre-init with all channels disabled.
+	OutputChannelMap.Empty();
+	OutputChannelMap.Init(IElectraDecoderAudioOutput::EChannelPosition::Disabled, MpegConfig.NumberOfChannels);
+	if (MpegConfig.NumberOfChannels == 1)
+	{
+		OutputChannelMap[0] = IElectraDecoderAudioOutput::EChannelPosition::C;
+	}
+	else if (MpegConfig.NumberOfChannels == 2)
+	{
+		OutputChannelMap[0] = IElectraDecoderAudioOutput::EChannelPosition::L;
+		OutputChannelMap[1] = IElectraDecoderAudioOutput::EChannelPosition::R;
+	}
+	else
+	{
+		// 2 channels at most.
+		return false;
+	}
+	return true;
+}
+
+
+
+
+IElectraDecoder::EDecoderError FElectraMP3Decoder::DecodeAccessUnit(const FInputAccessUnit& InInputAccessUnit, const TMap<FString, FVariant>& InAdditionalOptions)
+{
+	// If already in error do nothing!
+	if (LastError.IsSet())
+	{
+		return IElectraDecoder::EDecoderError::Error;
+	}
+
+	// Can not feed new input until draining has finished.
+	if (bFlushPending)
+	{
+		return IElectraDecoder::EDecoderError::EndOfData;
+	}
+
+	// If there is pending output it is very likely that decoding this access unit would also generate output.
+	// Since that would result in loss of the pending output we return now.
+	if (CurrentOutput.IsValid())
+	{
+		return IElectraDecoder::EDecoderError::NoBuffer;
+	}
+
+#if WITH_DR_LIBS_MP3
+	// Decode data.
+	if (InInputAccessUnit.Data && InInputAccessUnit.DataSize)
+	{
+		// Parse the codec specific information
+		if (!bHaveParsedMpegHeader)
+		{
+			if (!ParseMpegAudioHeader(MpegConfig, InInputAccessUnit.Data, InInputAccessUnit.DataSize))
+			{
+				return IElectraDecoder::EDecoderError::Error;
+			}
+			bHaveParsedMpegHeader = true;
+		}
+		// Set up the channel map accordingly.
+		if (!SetupChannelMap())
+		{
+			// Error was already posted.
+			return IElectraDecoder::EDecoderError::Error;
+		}
+		// Create decoder if necessary.
+		if (!DecoderHandle && !InternalDecoderCreate())
+		{
+			return IElectraDecoder::EDecoderError::Error;
+		}
+		// Decode
+		TSharedPtr<FElectraAudioDecoderOutputMP3_Common, ESPMode::ThreadSafe> NewOutput = MakeShared<FElectraAudioDecoderOutputMP3_Common>();
+		NewOutput->PTS = InInputAccessUnit.PTS;
+		NewOutput->UserValue = InInputAccessUnit.UserValue;
+		int32 AllocSize = sizeof(float) * MpegConfig.SamplesPerFrame * MpegConfig.NumberOfChannels;
+		NewOutput->Buffer = (float*)FMemory::Malloc(AllocSize);
+		NewOutput->NumChannels = MpegConfig.NumberOfChannels;
+		NewOutput->SampleRate = MpegConfig.SampleRate;
+		FMP3Decoder::FFrameInfo fi;
+		int Result = DecoderHandle->Decode(&fi, NewOutput->Buffer, AllocSize, (const unsigned char*) InInputAccessUnit.Data, (int) InInputAccessUnit.DataSize);
+		if (Result < 0)
+		{
+			PostError(Result, TEXT("FElectraMP3Decoder decoding failed"), ERRCODE_MP3DEC_INTERNAL_FAILED_TO_DECODE_INPUT);
+			return IElectraDecoder::EDecoderError::Error;
+		}
+		if (NewOutput->NumChannels != fi.NumChannels)
+		{
+			PostError(0, TEXT("Mismatching number of decoded channels during decode sequence!"), ERRCODE_MP3DEC_INTERNAL_FAILED_TO_DECODE_INPUT);
+			return IElectraDecoder::EDecoderError::Error;
+		}
+		if (NewOutput->SampleRate != fi.SampleRate)
+		{
+			PostError(0, TEXT("Mismatching sample rate during decode sequence!"), ERRCODE_MP3DEC_INTERNAL_FAILED_TO_DECODE_INPUT);
+			return IElectraDecoder::EDecoderError::Error;
+		}
+		if (Result > 0 && Result != MpegConfig.SamplesPerFrame)
+		{
+			PostError(0, TEXT("Mismatching samples per frame count during decode sequence!"), ERRCODE_MP3DEC_INTERNAL_FAILED_TO_DECODE_INPUT);
+			return IElectraDecoder::EDecoderError::Error;
+		}
+		if ((int32) InInputAccessUnit.DataSize != fi.NumFrameBytes)
+		{
+			PostError(Result, TEXT("FElectraMP3Decoder did not consume the entire input"), ERRCODE_MP3DEC_INTERNAL_FAILED_TO_DECODE_INPUT);
+			return IElectraDecoder::EDecoderError::Error;
+		}
+		NewOutput->NumFrames = Result;
+		NewOutput->ChannelPositions = OutputChannelMap;
+		CurrentOutput = MoveTemp(NewOutput);
+	}
+	return IElectraDecoder::EDecoderError::None;
+#else
+	return IElectraDecoder::EDecoderError::Error;
+#endif
+}
+
+IElectraDecoder::EDecoderError FElectraMP3Decoder::SendEndOfData()
+{
+	// If already in error do nothing!
+	if (LastError.IsSet())
+	{
+		return IElectraDecoder::EDecoderError::Error;
+	}
+	// Already draining?
+	if (bFlushPending)
+	{
+		return IElectraDecoder::EDecoderError::EndOfData;
+	}
+	bFlushPending = true;
+	return IElectraDecoder::EDecoderError::None;
+}
+
+IElectraDecoder::EDecoderError FElectraMP3Decoder::Flush()
+{
+	// If already in error do nothing!
+	if (LastError.IsSet())
+	{
+		return IElectraDecoder::EDecoderError::Error;
+	}
+	ResetToCleanStart();
+	return IElectraDecoder::EDecoderError::None;
+}
+
+IElectraDecoder::EOutputStatus FElectraMP3Decoder::HaveOutput()
+{
+	// If already in error do nothing!
+	if (LastError.IsSet())
+	{
+		return IElectraDecoder::EOutputStatus::Error;
+	}
+	// Have output?
+	if (CurrentOutput.IsValid())
+	{
+		return IElectraDecoder::EOutputStatus::Available;
+	}
+	// Pending flush?
+	if (bFlushPending)
+	{
+		bFlushPending = false;
+		return IElectraDecoder::EOutputStatus::EndOfData;
+	}
+	return IElectraDecoder::EOutputStatus::NeedInput;
+}
+
+TSharedPtr<IElectraDecoderOutput, ESPMode::ThreadSafe> FElectraMP3Decoder::GetOutput()
+{
+	TSharedPtr<IElectraDecoderOutput, ESPMode::ThreadSafe> Out = CurrentOutput;
+	CurrentOutput.Reset();
+	return Out;
+}
diff --git a/Engine/Plugins/Media/ElectraCodecs/Source/ElectraDecoders/Private/ElectraMediaMP3Decoder.h b/Engine/Plugins/Media/ElectraCodecs/Source/ElectraDecoders/Private/ElectraMediaMP3Decoder.h
new file mode 100644
--- /dev/null
+++ b/Engine/Plugins/Media/ElectraCodecs/Source/ElectraDecoders/Private/ElectraMediaMP3Decoder.h
@@ -0,0 +1,14 @@
+// Copyright Epic Games, Inc. All Rights Reserved.
+#pragma once
+
+#include "CoreMinimal.h"
+
+class IElectraCodecFactory;
+
+class FElectraMediaMP3Decoder
+{
+public:
+	static void Startup();
+	static void Shutdown();
+	static TSharedPtr<IElectraCodecFactory, ESPMode::ThreadSafe> CreateFactory();
+};
diff --git a/Engine/Plugins/Media/ElectraCodecs/Source/ElectraDecoders/Private/Utils/MPEG/ElectraUtilsMPEGAudio.cpp b/Engine/Plugins/Media/ElectraCodecs/Source/ElectraDecoders/Private/Utils/MPEG/ElectraUtilsMPEGAudio.cpp
--- a/Engine/Plugins/Media/ElectraCodecs/Source/ElectraDecoders/Private/Utils/MPEG/ElectraUtilsMPEGAudio.cpp
+++ b/Engine/Plugins/Media/ElectraCodecs/Source/ElectraDecoders/Private/Utils/MPEG/ElectraUtilsMPEGAudio.cpp
@@ -182,6 +182,167 @@ namespace ElectraDecodersUtil
 			//int64 remainingBits = bsp.GetRemainingBits();
 			return true;
 		}
+
+
+
+
+		namespace UtilsMPEG123
+		{
+			bool HasValidSync(uint32 InFrameHeader)
+			{
+				return (InFrameHeader & 0xffe00000U) == 0xffe00000U;
+			}
+
+			int32 GetVersionId(uint32 InFrameHeader)
+			{
+				check(HasValidSync(InFrameHeader));
+				return (int32) ((InFrameHeader >> 19) & 3U);
+			}
+
+			int32 GetLayerIndex(uint32 InFrameHeader)
+			{
+				check(HasValidSync(InFrameHeader));
+				return (int32) ((InFrameHeader >> 17) & 3U);
+			}
+
+			int32 GetBitrateIndex(uint32 InFrameHeader)
+			{
+				check(HasValidSync(InFrameHeader));
+				return (int32) ((InFrameHeader >> 12) & 15U);
+			}
+
+			int32 GetSamplingRateIndex(uint32 InFrameHeader)
+			{
+				check(HasValidSync(InFrameHeader));
+				return (int32) ((InFrameHeader >> 10) & 3U);
+			}
+
+			int32 GetChannelMode(uint32 InFrameHeader)
+			{
+				check(HasValidSync(InFrameHeader));
+				return (int32) ((InFrameHeader >> 6) & 3U);
+			}
+
+			int32 GetNumPaddingBytes(uint32 InFrameHeader)
+			{
+				check(HasValidSync(InFrameHeader));
+				return (int32) ((InFrameHeader >> 9) & 1U);
+			}
+
+			int32 GetVersion(uint32 InFrameHeader)
+			{
+				switch(GetVersionId(InFrameHeader))
+				{
+					default:
+						return 0;
+					case 0:
+						return 3;
+					case 2:
+						return 2;
+					case 3:
+						return 1;
+				}
+			}
+
+
+			int32 GetLayer(uint32 InFrameHeader)
+			{
+				switch(GetLayerIndex(InFrameHeader))
+				{
+					default:
+						return 0;
+					case 1:
+						return 3;
+					case 2:
+						return 2;
+					case 3:
+						return 1;
+				}
+			}
+
+			int32 GetBitrate(uint32 InFrameHeader)
+			{
+				static const int16 kBitrateTableMPEG1[3][16] =
+				{
+					{ 0, 32, 64, 96, 128, 160, 192, 224, 256, 288, 320, 352, 384, 416, 448, -1 },
+					{ 0, 32, 48, 56, 64, 80, 96, 112, 128, 160, 192, 224, 256, 320, 384, -1 },
+					{ 0, 32, 40, 48, 56, 64, 80, 96, 112, 128, 160, 192, 224, 256, 320, -1 }
+				};
+				static const int16 kBitrateTableMPEG2[3][16] =
+				{
+					{ 0, 32, 48, 56, 64, 80, 96, 112, 128, 144, 160, 176, 192, 224, 256, -1 },
+					{ 0, 8, 16, 24, 32, 40, 48, 56, 64, 80, 96, 112, 128, 144, 160 , -1 },
+					{ 0, 8, 16, 24, 32, 40, 48, 56, 64, 80, 96, 112, 128, 144, 160 , -1 }
+				};
+				int32 Version = GetVersion(InFrameHeader);
+				int32 Layer = GetLayer(InFrameHeader);
+				int32 BitrateIndex = GetBitrateIndex(InFrameHeader);
+				if (Version == 0 || Layer == 0 || BitrateIndex == 15)
+				{
+					return -1;
+				}
+				int32 bps = Version == 1 ? kBitrateTableMPEG1[Layer - 1][BitrateIndex] : kBitrateTableMPEG2[Layer - 1][BitrateIndex];
+				return bps * 1000;
+			}
+
+			int32 GetSamplingRate(uint32 InFrameHeader)
+			{
+				static const int32 kSamplingRates[4][4] =
+				{
+					{ -1, -1, -1, -1 },				// invalid
+					{ 44100, 48000, 32000, -1 },	// MPEG 1
+					{ 22050, 24000, 16000, -1 },	// MPEG 2
+					{ 11025, 12000,  8000, -1 }		// MPEG 2.5
+				};
+				int32 Version = GetVersion(InFrameHeader);
+				int32 SampleRateIndex = GetSamplingRateIndex(InFrameHeader);
+				return kSamplingRates[Version][SampleRateIndex];
+			}
+			int32 GetChannelCount(uint32 InFrameHeader)
+			{
+				return GetChannelMode(InFrameHeader) == 3 ? 1 : 2;
+			}
+
+			int32 GetFrameSize(uint32 InFrameHeader, int32 InForcedPadding)
+			{
+				static const int32 kNumCoeffs[2][3] =
+				{
+					{ 12, 144, 144 },		// MPEG 1 (layer 1, 2, 3)
+					{ 12, 144, 72 }			// MPEG 2 / 2.5 (layer 1, 2, 3)
+				};
+				static const int32 kSlotSize[3] =
+				{
+					4, 1, 1					// Layer 1, 2, 3
+				};
+
+				int32 NumPadding = InForcedPadding < 0 ? GetNumPaddingBytes(InFrameHeader) : InForcedPadding;
+				int32 Version = GetVersion(InFrameHeader);
+				int32 Layer = GetLayer(InFrameHeader);
+				int32 Bitrate = GetBitrate(InFrameHeader);
+				int32 SampleRate = GetSamplingRate(InFrameHeader);
+				if (Version == 0 || Layer == 0 || Bitrate <= 0 || SampleRate <= 0)
+				{
+					return 0;
+				}
+				int32 FrameSize = (kNumCoeffs[Version==1?0:1][Layer-1] * Bitrate / SampleRate + NumPadding) * kSlotSize[Layer-1];
+				return FrameSize;
+			}
+
+			int32 GetSamplesPerFrame(uint32 InFrameHeader)
+			{
+				static const int32 kSamplesPerFrame[2][4] =
+				{
+					{ 0, 384, 1152, 1152 },	// MPEG 1 (layer 1, 2, 3)
+					{ 0, 384, 1152, 576 }	// MPEG 2 / 2.5 (layer 1, 2, 3)
+				};
+				int32 Version = GetVersion(InFrameHeader);
+				int32 Layer = GetLayer(InFrameHeader);
+				return kSamplesPerFrame[Version==1?0:1][Layer];
+			}
+
+		} // namespace UtilsMPEG123
+
+
 	} // namespace MPEG
 } // namespace ElectraDecodersUtil
 
diff --git a/Engine/Plugins/Media/ElectraCodecs/Source/ElectraDecoders/Public/Utils/MPEG/ElectraUtilsMPEGAudio.h b/Engine/Plugins/Media/ElectraCodecs/Source/ElectraDecoders/Public/Utils/MPEG/ElectraUtilsMPEGAudio.h
--- a/Engine/Plugins/Media/ElectraCodecs/Source/ElectraDecoders/Public/Utils/MPEG/ElectraUtilsMPEGAudio.h
+++ b/Engine/Plugins/Media/ElectraCodecs/Source/ElectraDecoders/Public/Utils/MPEG/ElectraUtilsMPEGAudio.h
@@ -49,6 +49,29 @@ namespace ElectraDecodersUtil
 		}
 
 
+		namespace UtilsMPEG123
+		{
+			// Raw bit methods
+			bool ELECTRADECODERS_API HasValidSync(uint32 InFrameHeader);
+			int32 ELECTRADECODERS_API GetVersionId(uint32 InFrameHeader);			// 0=MPEG2.5, 1=reserved, 2=MPEG2 (ISO/IEC 13818-3), 3=MPEG1 (ISO/IEC 11172-3)
+			int32 ELECTRADECODERS_API GetLayerIndex(uint32 InFrameHeader);			// 0=reserved, 1=Layer III, 2=Layer II, 3=Layer I
+			int32 ELECTRADECODERS_API GetBitrateIndex(uint32 InFrameHeader);		// 0-15
+			int32 ELECTRADECODERS_API GetSamplingRateIndex(uint32 InFrameHeader);	// 0-3
+			int32 ELECTRADECODERS_API GetChannelMode(uint32 InFrameHeader);			// 0-3
+			int32 ELECTRADECODERS_API GetNumPaddingBytes(uint32 InFrameHeader);		// 0 or 1
+
+
+			// Convenience methods
+			int32 ELECTRADECODERS_API GetVersion(uint32 InFrameHeader);				// 1=MPEG 1, 2=MPEG2, 3=MPEG2.5, 0=reserved
+			int32 ELECTRADECODERS_API GetLayer(uint32 InFrameHeader);				// 0=reserved, 1=Layer I, 2=Layer II, 3=Layer III
+			int32 ELECTRADECODERS_API GetBitrate(uint32 InFrameHeader);				// Bitrate in kbps, -1=invalid
+			int32 ELECTRADECODERS_API GetSamplingRate(uint32 InFrameHeader);		// Sampling rate, -1=invalid
+			int32 ELECTRADECODERS_API GetChannelCount(uint32 InFrameHeader);		// 1=mono, 2=stereo
+
+			int32 ELECTRADECODERS_API GetSamplesPerFrame(uint32 InFrameHeader);		// number of samples encoded in the frame
+			int32 ELECTRADECODERS_API GetFrameSize(uint32 InFrameHeader, int32 InForcedPadding=-1);		// number of bytes in the packet if CBR encoded, 0=could not calculate
+		}
+
 	} // namespace MPEG
 
 } // namespace ElectraDecodersUtil
diff --git a/Engine/Plugins/Media/ElectraPlayer/Source/ElectraPlayerRuntime/Private/Runtime/Demuxer/ParserISO14496-12.cpp b/Engine/Plugins/Media/ElectraPlayer/Source/ElectraPlayerRuntime/Private/Runtime/Demuxer/ParserISO14496-12.cpp
--- a/Engine/Plugins/Media/ElectraPlayer/Source/ElectraPlayerRuntime/Private/Runtime/Demuxer/ParserISO14496-12.cpp
+++ b/Engine/Plugins/Media/ElectraPlayer/Source/ElectraPlayerRuntime/Private/Runtime/Demuxer/ParserISO14496-12.cpp
@@ -55,7 +55,7 @@ namespace Electra
 			for(uint32 i=0, Atom=In4CC; i<4; ++i, Atom<<=8)
 			{
 				int32 v = Atom >> 24;
-				if ((v >= 'A' && v <= 'Z') || (v >= 'a' && v <= 'z') || (v >= '0' && v <= '9') || v == '_')
+				if ((v >= 'A' && v <= 'Z') || (v >= 'a' && v <= 'z') || (v >= '0' && v <= '9') || v == '_' || v == '.')
 				{
 					Out.AppendChar(v);
 				}
@@ -7460,7 +7460,7 @@ namespace Electra
 
 	UEMediaError FParserISO14496_12::ParseMP4ASampleType(IPlayerSessionServices* PlayerSession, FTrack* Track, const FMP4Box* SampleBox)
 	{
-		auto HandleESDS = [PlayerSession, Track](const FMP4BoxESDS* ESDSBox, bool& bGotAudioFormat, bool& bIsSupported) -> bool
+		auto HandleESDS = [PlayerSession, Track](const FMP4BoxAudioSampleEntry* AudioSampleEntry, const FMP4BoxESDS* ESDSBox, bool& bGotAudioFormat, bool& bIsSupported) -> bool
 		{
 			if (!ESDSBox)
 			{
@@ -7513,6 +7513,37 @@ namespace Electra
 						}
 					}
 				}
+				// Is this MPEG-1 audio?
+				else if (Track->CodecSpecificDataMP4A.GetObjectTypeID() == MPEG::FESDescriptor::FObjectTypeID::MPEG1_Audio &&
+					Track->CodecSpecificDataMP4A.GetStreamType() == MPEG::FESDescriptor::FStreamType::AudioStream)
+				{
+					Track->CodecInformation.SetStreamType(EStreamType::Audio);
+					Track->CodecInformation.SetCodec(FStreamCodecInformation::ECodec::Audio4CC);
+					Track->CodecInformation.SetCodec4CC(Utils::Make4CC('m','p','g','a'));
+					Track->CodecInformation.SetProfile(1);
+					Track->CodecInformation.SetCodecSpecificData(Track->CodecSpecificDataRAW);
+					Track->CodecInformation.SetStreamLanguageCode(Track->GetLanguage());
+					Track->CodecInformation.SetCodecSpecifierRFC6381(FString::Printf(TEXT("mp4a.6b")));
+					Track->CodecInformation.SetSamplingRate((int32) AudioSampleEntry->GetSampleRate());
+					Track->CodecInformation.SetNumberOfChannels(AudioSampleEntry->GetChannelCount());
+
+					// Typically an mp4a track will not have a 'btrt' box because the bitrate is stored in the DecoderConfigDescriptor.
+					Track->BitrateInfo.BufferSizeDB = Track->CodecSpecificDataMP4A.GetBufferSize();
+					Track->BitrateInfo.MaxBitrate = Track->CodecSpecificDataMP4A.GetMaxBitrate();
+					Track->BitrateInfo.AvgBitrate = Track->CodecSpecificDataMP4A.GetAvgBitrate();
+					Track->CodecInformation.SetBitrate(Track->BitrateInfo.MaxBitrate);
+
+					bGotAudioFormat = true;
+
+					if (PlayerSession)
+					{
+						IPlayerStreamFilter* StreamFilter = PlayerSession->GetStreamFilter();
+						if (StreamFilter && !StreamFilter->CanDecodeStream(Track->CodecInformation))
+						{
+							bIsSupported = false;
+						}
+					}
+				}
 				else
 				{
 					// Not expected format. Ignore this track.
@@ -7550,7 +7581,7 @@ namespace Electra
 					case FMP4Box::kBox_esds:
 					{
 						const FMP4BoxESDS* ESDSBox = static_cast<const FMP4BoxESDS*>(AudioSampleEntry->GetChildBox(0));
-						bool bOk = HandleESDS(ESDSBox, bGotAudioFormat, bIsSupported);
+						bool bOk = HandleESDS(AudioSampleEntry, ESDSBox, bGotAudioFormat, bIsSupported);
 						if (!bOk)
 						{
 							return UEMEDIA_ERROR_FORMAT_ERROR;
@@ -7560,7 +7591,7 @@ namespace Electra
 					case FMP4Box::kBox_wave:
 					{
 						const FMP4BoxESDS* ESDSBox = static_cast<const FMP4BoxESDS*>(AudioSampleEntry->FindBox(FMP4Box::kBox_esds, 1));
-						bool bOk = HandleESDS(ESDSBox, bGotAudioFormat, bIsSupported);
+						bool bOk = HandleESDS(AudioSampleEntry, ESDSBox, bGotAudioFormat, bIsSupported);
 						if (!bOk)
 						{
 							return UEMEDIA_ERROR_FORMAT_ERROR;
diff --git a/Engine/Plugins/Media/ElectraPlayer/Source/ElectraPlayerRuntime/Private/Runtime/Demuxer/ParserMKV.cpp b/Engine/Plugins/Media/ElectraPlayer/Source/ElectraPlayerRuntime/Private/Runtime/Demuxer/ParserMKV.cpp
--- a/Engine/Plugins/Media/ElectraPlayer/Source/ElectraPlayerRuntime/Private/Runtime/Demuxer/ParserMKV.cpp
+++ b/Engine/Plugins/Media/ElectraPlayer/Source/ElectraPlayerRuntime/Private/Runtime/Demuxer/ParserMKV.cpp
@@ -864,9 +864,9 @@ namespace Electra
 		public:
 			FMKVSeekHead(uint32 InID, int64 InByteOffset, int64 InByteSize) : FMKVEBMLElement(InID, InByteOffset, InByteSize) { }
 			virtual ~FMKVSeekHead() = default;
-			virtual EParseResult ParseElement(IMKVFetcher* InReader) 
-			{ 
-				EParseResult Result = ParseElementList(Elements, InReader); 
+			virtual EParseResult ParseElement(IMKVFetcher* InReader)
+			{
+				EParseResult Result = ParseElementList(Elements, InReader);
 				// We sort the Seek elements by their file position. This may avoid too many random file accesses
 				// when loading the elements after a cluster.
 				Seeks.Sort([](const TMKVElementPtr<FMKVSeek>& a, const TMKVElementPtr<FMKVSeek>& b){return a->GetSeekPosition() < b->GetSeekPosition();});
@@ -1166,9 +1166,9 @@ namespace Electra
 		public:
 			FMKVAudio(uint32 InID, int64 InByteOffset, int64 InByteSize) : FMKVEBMLElement(InID, InByteOffset, InByteSize) { }
 			virtual ~FMKVAudio() = default;
-			virtual EParseResult ParseElement(IMKVFetcher* InReader) 
-			{ 
-				EParseResult Result = ParseElementList(Elements, InReader); 
+			virtual EParseResult ParseElement(IMKVFetcher* InReader)
+			{
+				EParseResult Result = ParseElementList(Elements, InReader);
 				if (OutputSamplingFrequency == 0.0)
 				{
 					OutputSamplingFrequency = SamplingFrequency;
@@ -1324,8 +1324,8 @@ namespace Electra
 			FMKVContentEncodings(uint32 InID, int64 InByteOffset, int64 InByteSize) : FMKVEBMLElement(InID, InByteOffset, InByteSize) { }
 			virtual ~FMKVContentEncodings() = default;
 			virtual EParseResult ParseElement(IMKVFetcher* InReader)
-			{ 
-				EParseResult Result = ParseElementList(Elements, InReader); 
+			{
+				EParseResult Result = ParseElementList(Elements, InReader);
 				// The content encodings must be processed in descending `ContentEncodingOrder`, but it is not mandated
 				// that they are sorted such in the file. We sort them now.
 				ContentEncoding.StableSort([](const TMKVElementPtr<FMKVContentEncoding>& a, const TMKVElementPtr<FMKVContentEncoding>& b){return a->GetContentEncodingOrder() > b->GetContentEncodingOrder();});
@@ -1540,8 +1540,8 @@ namespace Electra
 		public:
 			FMKVChapters(uint32 InID, int64 InByteOffset, int64 InByteSize) : FMKVEBMLElement(InID, InByteOffset, InByteSize) { }
 			virtual ~FMKVChapters() = default;
-			virtual EParseResult ParseElement(IMKVFetcher* InReader) 
-			{ 
+			virtual EParseResult ParseElement(IMKVFetcher* InReader)
+			{
 				// Not using chapters. Skip over them.
 				return InReader->FetchSkipOver(ByteSize) ? EParseResult::Ok : EParseResult::Error;
 			}
@@ -1782,7 +1782,7 @@ namespace Electra
 			FMKVCues(uint32 InID, int64 InByteOffset, int64 InByteSize) : FMKVEBMLElement(InID, InByteOffset, InByteSize) { }
 			virtual ~FMKVCues() = default;
 			virtual EParseResult ParseElement(IMKVFetcher* InReader)
-			{ 
+			{
 				EParseResult Result = ParseElementList(Elements, InReader);
 				if (Result == EParseResult::Ok)
 				{
@@ -2022,7 +2022,7 @@ namespace Electra
 
 	public:
 		class FMKVTrack;
-		
+
 		class FMKVCueIterator : public IParserMKV::ICueIterator
 		{
 		public:
@@ -2449,6 +2449,9 @@ namespace Electra
 			{ TEXT("A_AAC/MPEG4/LC"), TEXT("aac") },
 			{ TEXT("A_AAC/MPEG4/LC/SBR"), TEXT("aac") },
 			{ TEXT("A_OPUS"), TEXT("Opus") },
+			{ TEXT("A_MPEG/L3"), TEXT("MPEG") },
+			{ TEXT("A_MPEG/L2"), TEXT("MPEG") },
+			{ TEXT("A_MPEG/L1"), TEXT("MPEG") },
 			//{ TEXT("S_TEXT/UTF8"), TEXT("") },
 			//{ TEXT("S_TEXT/WEBVTT"), TEXT("") },
 			//{ TEXT(""), TEXT("") },
@@ -2932,8 +2935,8 @@ namespace Electra
 			TArray<uint8> vpcC { 1,0,0,0, 0,0,0,0, 0,0,0,0 };
 			vpcC[4] = (uint8)OutCodecInformation.GetProfile();
 			vpcC[5] = (uint8)OutCodecInformation.GetProfileLevel();
-			vpcC[6] = (uint8)((OutCodecInformation.GetCodecVideoColorInfo().BitDepthLuma.Get(8) << 4) + 
-							  (OutCodecInformation.GetCodecVideoColorInfo().ChromaSubsampling.Get(0) << 1) + 
+			vpcC[6] = (uint8)((OutCodecInformation.GetCodecVideoColorInfo().BitDepthLuma.Get(8) << 4) +
+							  (OutCodecInformation.GetCodecVideoColorInfo().ChromaSubsampling.Get(0) << 1) +
 							   OutCodecInformation.GetCodecVideoColorInfo().VideoFullRangeFlag.Get(0));
 			vpcC[7] = (uint8)OutCodecInformation.GetCodecVideoColorInfo().ColourPrimaries.Get(2);
 			vpcC[8] = (uint8)OutCodecInformation.GetCodecVideoColorInfo().TransferCharacteristics.Get(2);
@@ -3076,8 +3079,8 @@ namespace Electra
 			TArray<uint8> vpcC { 1,0,0,0, 0,0,0,0, 0,0,0,0 };
 			vpcC[4] = (uint8)OutCodecInformation.GetProfile();
 			vpcC[5] = (uint8)OutCodecInformation.GetProfileLevel();
-			vpcC[6] = (uint8)((OutCodecInformation.GetCodecVideoColorInfo().BitDepthLuma.Get(8) << 4) + 
-							  (OutCodecInformation.GetCodecVideoColorInfo().ChromaSubsampling.Get(0) << 1) + 
+			vpcC[6] = (uint8)((OutCodecInformation.GetCodecVideoColorInfo().BitDepthLuma.Get(8) << 4) +
+							  (OutCodecInformation.GetCodecVideoColorInfo().ChromaSubsampling.Get(0) << 1) +
 							   OutCodecInformation.GetCodecVideoColorInfo().VideoFullRangeFlag.Get(0));
 			vpcC[7] = (uint8)OutCodecInformation.GetCodecVideoColorInfo().ColourPrimaries.Get(2);
 			vpcC[8] = (uint8)OutCodecInformation.GetCodecVideoColorInfo().TransferCharacteristics.Get(2);
@@ -3160,7 +3163,7 @@ namespace Electra
 			// The codec private data is (presumably) an `OpusHead` structure as described here
 			//   https://datatracker.ietf.org/doc/html/rfc7845#section-5.1
 			// according to the Matroska Opus *DRAFT* described here: https://wiki.xiph.org/MatroskaOpus
-			
+
 			TArray<uint8> OpusHead = InFromTrack->GetCodecPrivate();
 			const TArray<uint8> MagicOpusHeader {'O','p','u','s','H','e','a','d'};
 			if (OpusHead.Num() < 8 || FMemory::Memcmp(OpusHead.GetData(), MagicOpusHeader.GetData(), MagicOpusHeader.Num()))
@@ -3202,12 +3205,10 @@ namespace Electra
 			OutCodecInformation.SetCodecSpecifierRFC6381(TEXT("Opus"));
 			OutCodecInformation.SetSamplingRate(Audio->GetOutputSampleRate());
 			OutCodecInformation.SetNumberOfChannels(Audio->GetNumberOfChannels());
-#if 0
-			// If there is no default duration set we try to calculate it from the sample rate.
-			const int32 NumDecodedSamplesPerBlock = 960;
+			// If there is no default duration set we assume the encoded frame size was 20ms
 			if (InFromTrack->GetDefaultDurationNanos() == 0)
 			{
-				FTimeFraction fr(NumDecodedSamplesPerBlock, (uint32)OutCodecInformation.GetSamplingRate());
+				FTimeFraction fr(20, 1000);
 				if (fr.IsValid())
 				{
 					int64 nanos = fr.GetAsTimebase(1000000000);
@@ -3218,22 +3219,39 @@ namespace Electra
 					}
 				}
 			}
-#else
-			// If there is no default duration set we assume the encoded frame size was 20ms
+			return true;
+		}
+		// MPEG audio?
+		else if (InFromTrack->GetTrackType() == EMKVTrackType::Audio && Codec4CC.Equals(TEXT("MPEG")))
+		{
+			TMKVElementPtr<FMKVAudio> Audio = InFromTrack->GetAudio();
+			if (!Audio.IsValid())
+			{
+				return false;
+			}
+
+			OutCodecInformation.SetStreamType(EStreamType::Audio);
+			OutCodecInformation.SetMimeType(TEXT("audio/mpeg"));
+			OutCodecInformation.SetCodec(FStreamCodecInformation::ECodec::Audio4CC);
+			OutCodecInformation.SetCodec4CC(Utils::Make4CC('m','p','g','a'));
+			OutCodecInformation.SetProfile(1);
+			FString Layer = InFromTrack->GetCodecID().Mid(8);
+			int32 LayerValue = 3;
+			if (Layer.Len())
+			{
+				LexFromString(LayerValue, *Layer);
+			}
+			OutCodecInformation.SetProfileLevel(LayerValue);
+			OutCodecInformation.SetStreamLanguageCode(InFromTrack->GetLanguage());
+			OutCodecInformation.SetCodecSpecifierRFC6381(TEXT("mp4a.6b"));
+			OutCodecInformation.SetSamplingRate(Audio->GetOutputSampleRate());
+			OutCodecInformation.SetNumberOfChannels(Audio->GetNumberOfChannels());
+			// If there is no default duration set we try to calculate it from the sample rate.
 			if (InFromTrack->GetDefaultDurationNanos() == 0)
 			{
-				FTimeFraction fr(20, 1000);
-				if (fr.IsValid())
-				{
-					int64 nanos = fr.GetAsTimebase(1000000000);
-					check(nanos >= 0);
-					if (nanos >= 0)
-					{
-						InFromTrack->SetDefaultDurationNanos((uint64)nanos);
-					}
-				}
+				UE_LOG(LogElectraMKVParser, Error, TEXT("A_MPEG requires a sample default duration specified in the track!"));
+				return false;
 			}
-#endif
 			return true;
 		}
 		return false;
@@ -3840,7 +3858,7 @@ namespace Electra
 					{
 						int32 SampleSize = 0;
 						uint8 b;
-						do 
+						do
 						{
 							if (!Reader->Read(b))
 							{
diff --git a/Engine/Plugins/Media/ElectraPlayer/Source/ElectraPlayerRuntime/Private/Runtime/StreamTypes.cpp b/Engine/Plugins/Media/ElectraPlayer/Source/ElectraPlayerRuntime/Private/Runtime/StreamTypes.cpp
--- a/Engine/Plugins/Media/ElectraPlayer/Source/ElectraPlayerRuntime/Private/Runtime/StreamTypes.cpp
+++ b/Engine/Plugins/Media/ElectraPlayer/Source/ElectraPlayerRuntime/Private/Runtime/StreamTypes.cpp
@@ -26,7 +26,7 @@ namespace Electra
 			for(uint32 i=0, Atom=In4CC; i<4; ++i, Atom<<=8)
 			{
 				int32 v = Atom >> 24;
-				if ((v >= 'A' && v <= 'Z') || (v >= 'a' && v <= 'z') || (v >= '0' && v <= '9') || v == '_')
+				if ((v >= 'A' && v <= 'Z') || (v >= 'a' && v <= 'z') || (v >= '0' && v <= '9') || v == '_'|| v == '.')
 				{
 					Out.AppendChar(v);
 				}
@@ -490,8 +490,22 @@ namespace Electra
 				switch(GetCodec4CC())
 				{
 					case Make4CC('O','p','u','s'):
+					{
 						HumanReadableCodecName = TEXT("Opus");
 						return true;
+					}
+					case Make4CC('m','p','g','a'):
+					{
+						if (GetProfileLevel())
+						{
+							HumanReadableCodecName = FString::Printf(TEXT("MPEG%d Layer %d"), GetProfile(), GetProfileLevel());
+						}
+						else
+						{
+							HumanReadableCodecName = FString::Printf(TEXT("MPEG%d audio"), GetProfile());
+						}
+						return true;
+					}
 				}
 				HumanReadableCodecName = Printable4CC(GetCodec4CC());
 				return true;
diff --git a/Engine/Plugins/Media/ElectraPlayer/Source/ElectraPlayerRuntime/Private/Runtime/Utilities/UtilsMPEG.h b/Engine/Plugins/Media/ElectraPlayer/Source/ElectraPlayerRuntime/Private/Runtime/Utilities/UtilsMPEG.h
--- a/Engine/Plugins/Media/ElectraPlayer/Source/ElectraPlayerRuntime/Private/Runtime/Utilities/UtilsMPEG.h
+++ b/Engine/Plugins/Media/ElectraPlayer/Source/ElectraPlayerRuntime/Private/Runtime/Utilities/UtilsMPEG.h
@@ -43,7 +43,8 @@ namespace Electra
 				H264 = 0x21,
 				H264_ParameterSets = 0x22,
 				H265 = 0x23,
-				MPEG4_Audio = 0x40
+				MPEG4_Audio = 0x40,
+				MPEG1_Audio = 0x6b
 			};
 			enum class FStreamType
 			{
