ElectraDecoders: Raised Windows H.264 video decoder limit from UHD to 4K

#rnx

diff --git a/Engine/Plugins/Media/ElectraCodecs/Source/ElectraDecoders/Private/Windows/h264/H264_VideoDecoder_Windows.cpp b/Engine/Plugins/Media/ElectraCodecs/Source/ElectraDecoders/Private/Windows/h264/H264_VideoDecoder_Windows.cpp
--- a/Engine/Plugins/Media/ElectraCodecs/Source/ElectraDecoders/Private/Windows/h264/H264_VideoDecoder_Windows.cpp
+++ b/Engine/Plugins/Media/ElectraCodecs/Source/ElectraDecoders/Private/Windows/h264/H264_VideoDecoder_Windows.cpp
@@ -162,7 +162,7 @@ namespace IElectraVideoDecoderH264_DX_Platform
 		{ return Ctx.DxDeviceContext.GetReference(); }
 		virtual bool IsSoftware() const override
 		{ return bIsSW; }
-		
+
 		void Close()
 		{
 			MFT.SafeRelease();
@@ -191,11 +191,11 @@ void IElectraVideoDecoderH264_DX::PlatformGetSupportedConfigurations(TArray<FSup
 	{
 		IElectraVideoDecoderH264_DX_Platform::DecoderConfigurations.Empty();
 		// Baseline
-		IElectraVideoDecoderH264_DX_Platform::DecoderConfigurations.Emplace(IElectraVideoDecoderH264_DX::FSupportedConfiguration(66, 52, 120, 3840, 2160, 0 /*1280*720 / 256*/));
+		IElectraVideoDecoderH264_DX_Platform::DecoderConfigurations.Emplace(IElectraVideoDecoderH264_DX::FSupportedConfiguration(66, 52, 120, 4096, 2304, 0 /*1280*720 / 256*/));
 		// Main
-		IElectraVideoDecoderH264_DX_Platform::DecoderConfigurations.Emplace(IElectraVideoDecoderH264_DX::FSupportedConfiguration(77, 52, 120, 3840, 2160, 0 /*1280*720 / 256*/));
+		IElectraVideoDecoderH264_DX_Platform::DecoderConfigurations.Emplace(IElectraVideoDecoderH264_DX::FSupportedConfiguration(77, 52, 120, 4096, 2304, 0 /*1280*720 / 256*/));
 		// High
-		IElectraVideoDecoderH264_DX_Platform::DecoderConfigurations.Emplace(IElectraVideoDecoderH264_DX::FSupportedConfiguration(100, 52, 120, 3840, 2160, 0 /*1280*720 / 256*/));
+		IElectraVideoDecoderH264_DX_Platform::DecoderConfigurations.Emplace(IElectraVideoDecoderH264_DX::FSupportedConfiguration(100, 52, 120, 4096, 2304, 0 /*1280*720 / 256*/));
 
 		IElectraVideoDecoderH264_DX_Platform::bDecoderConfigurationsDirty = false;
 	}
