ElectraDecoders: Fixed H.265 frame alignment in D3D12 decoder

#rnx

diff --git a/Engine/Plugins/Media/D3D12VideoDecodersElectra/Source/Private/VideoDecoder_D3D12_H265.cpp b/Engine/Plugins/Media/D3D12VideoDecodersElectra/Source/Private/VideoDecoder_D3D12_H265.cpp
--- a/Engine/Plugins/Media/D3D12VideoDecodersElectra/Source/Private/VideoDecoder_D3D12_H265.cpp
+++ b/Engine/Plugins/Media/D3D12VideoDecodersElectra/Source/Private/VideoDecoder_D3D12_H265.cpp
@@ -270,8 +270,7 @@ IElectraDecoder::EDecoderError FD3D12VideoDecoder_H265::DecodeAccessUnit(const F
 			}
 			if (!CurrentConfig.VideoDecoderHeap.IsValid())
 			{
-				const int32 Alignment = spsPtr->GetMinCbSizeY();
-				if (!CreateDecoderHeap(DPBSize, dw, dh, Alignment))
+				if (!CreateDecoderHeap(DPBSize, dw, dh, spsPtr->GetMinCbSizeY()))
 				{
 					return IElectraDecoder::EDecoderError::Error;
 				}
@@ -279,15 +278,12 @@ IElectraDecoder::EDecoderError FD3D12VideoDecoder_H265::DecodeAccessUnit(const F
 
 			if (!DPB.IsValid())
 			{
-				// Use the maximum MinCbSizeY value for image alignment as stipulated in the DXVA HEVC documentation.
-				const int32 Alignment = 64;
-
 				// As far as the decoded frames go, their size can be the maximum that is required
 				// for this stream (the largest resolution).
 				const int32 Width = (int32) DecodeSupport.Width;
 				const int32 Height = (int32) DecodeSupport.Height;
 				const int32 NumFrames = spsPtr->GetDPBSize() + 2;	// 1 extra for the current frame that's not in the DPB yet, and 1 extra that acts as a 'missing' frame.
-				if (!CreateDPB(DPB, Width, Height, Alignment, NumFrames))
+				if (!CreateDPB(DPB, Width, Height, GetFrameAlignment(), NumFrames))
 				{
 					return IElectraDecoder::EDecoderError::Error;
 				}
@@ -771,8 +767,9 @@ IElectraDecoder::EDecoderError FD3D12VideoDecoder_H265::DecodeSlicesH265(const F
 	InDec->bDoNotOutput = (InInputAccessUnit.Flags & EElectraDecoderFlags::DoNotOutput) != EElectraDecoderFlags::None;
 	InDec->OutputType = IElectraDecoderVideoOutput::EOutputType::Output;
 	InSequenceParameterSet.GetCrop(InDec->Crop.Left, InDec->Crop.Right, InDec->Crop.Top, InDec->Crop.Bottom);
-	InDec->Width = InSequenceParameterSet.GetWidth();
-	InDec->Height = InSequenceParameterSet.GetHeight();
+
+	InDec->Width = Align(InSequenceParameterSet.GetWidth(), GetFrameAlignment());
+	InDec->Height = Align(InSequenceParameterSet.GetHeight(), GetFrameAlignment());
 	InDec->ImageWidth = InDec->Width - InDec->Crop.Left - InDec->Crop.Right;
 	InDec->ImageHeight = InDec->Height - InDec->Crop.Top - InDec->Crop.Bottom;
 	//InDec->Pitch = InDec->ImageWidth;
diff --git a/Engine/Plugins/Media/D3D12VideoDecodersElectra/Source/Private/VideoDecoder_D3D12_H265.h b/Engine/Plugins/Media/D3D12VideoDecodersElectra/Source/Private/VideoDecoder_D3D12_H265.h
--- a/Engine/Plugins/Media/D3D12VideoDecodersElectra/Source/Private/VideoDecoder_D3D12_H265.h
+++ b/Engine/Plugins/Media/D3D12VideoDecodersElectra/Source/Private/VideoDecoder_D3D12_H265.h
@@ -66,6 +66,9 @@ private:
 	FBitstreamParamsH265 BitstreamParamsH265;
 	static const uint8 ScanOrderDiag4[16];
 	static const uint8 ScanOrderDiag8[64];
+	// Use the maximum MinCbSizeY value for image alignment as stipulated in the DXVA HEVC documentation.
+	constexpr int32 GetFrameAlignment()
+	{ return 64; }
 };
 
 }
diff --git a/Engine/Plugins/Media/ElectraCodecs/Source/ElectraDecoders/Private/DX/VideoDecoderH265_DX.cpp b/Engine/Plugins/Media/ElectraCodecs/Source/ElectraDecoders/Private/DX/VideoDecoderH265_DX.cpp
--- a/Engine/Plugins/Media/ElectraCodecs/Source/ElectraDecoders/Private/DX/VideoDecoderH265_DX.cpp
+++ b/Engine/Plugins/Media/ElectraCodecs/Source/ElectraDecoders/Private/DX/VideoDecoderH265_DX.cpp
@@ -16,7 +16,8 @@
 #include COMPILED_PLATFORM_HEADER(ElectraDecoderGPUBufferHelpers.h)
 #include COMPILED_PLATFORM_HEADER(PlatformHeaders_Video_DX.h)
 
-#include "Stats\Stats.h"
+#include "Stats/Stats.h"
+#include "HAL/IConsoleManager.h"
 
 /*********************************************************************************************************************/
 #include COMPILED_PLATFORM_HEADER(PlatformHeaders_Video_DX.h)
@@ -39,6 +40,15 @@ if (FAILED(Result))						\
 
 DECLARE_CYCLE_STAT(TEXT("ElectraDecoder ConvertOutput"), STAT_ElectraDecoder_ConvertOutputH265, STATGROUP_Media);
 
+/*********************************************************************************************************************/
+
+static TAutoConsoleVariable<int32> CVarElectraWindowsH265UseOldOutputPath(
+	TEXT("Electra.Win.H265UseOldOutputPath"),
+	0,
+	TEXT("Use old style CPU buffer output path for H265.\n")
+	TEXT(" 0: use new upload heap / direct path (default); 1: use old CPU buffer path."),
+	ECVF_Default);
+
 /*********************************************************************************************************************/
 /*********************************************************************************************************************/
 /*********************************************************************************************************************/
@@ -62,62 +72,34 @@ public:
 	{ }
 
 	FTimespan GetPTS() const override
-	{
-		return PTS;
-	}
+	{ return PTS; }
 	uint64 GetUserValue() const override
-	{
-		return UserValue;
-	}
+	{ return UserValue; }
 
 	EOutputType GetOutputType() const
-	{
-		return EOutputType::Output;
-	}
+	{ return EOutputType::Output; }
 	int32 GetWidth() const override
-	{
-		return Width;
-	}
+	{ return Width; }
 	int32 GetHeight() const override
-	{
-		return Height;
-	}
+	{ return Height; }
 	int32 GetDecodedWidth() const override
-	{
-		return DecodedWidth;
-	}
+	{ return DecodedWidth; }
 	int32 GetDecodedHeight() const override
-	{
-		return DecodedHeight;
-	}
+	{ return DecodedHeight; }
 	FElectraVideoDecoderOutputCropValues GetCropValues() const override
-	{
-		return Crop;
-	}
+	{ return Crop; }
 	int32 GetAspectRatioW() const override
-	{
-		return AspectW;
-	}
+	{ return AspectW; }
 	int32 GetAspectRatioH() const override
-	{
-		return AspectH;
-	}
+	{ return AspectH; }
 	int32 GetFrameRateNumerator() const override
-	{
-		return FrameRateN;
-	}
+	{ return FrameRateN; }
 	int32 GetFrameRateDenominator() const override
-	{
-		return FrameRateD;
-	}
+	{ return FrameRateD; }
 	int32 GetNumberOfBits() const override
-	{
-		return NumBits;
-	}
+	{ return NumBits; }
 	void GetExtraValues(TMap<FString, FVariant>& OutExtraValues) const override
-	{
-		OutExtraValues = ExtraValues;
-	}
+	{ OutExtraValues = ExtraValues; }
 	void* GetPlatformOutputHandle(EElectraDecoderPlatformOutputHandleType InTypeOfHandle) const override
 	{
 		switch (InTypeOfHandle)
@@ -144,13 +126,9 @@ public:
 		}
 	}
 	IElectraDecoderVideoOutputTransferHandle* GetTransferHandle() const override
-	{
-		return nullptr;
-	}
+	{ return nullptr; }
 	IElectraDecoderVideoOutput::EImageCopyResult CopyPlatformImage(IElectraDecoderVideoOutputCopyResources* InCopyResources) const override
-	{
-		return IElectraDecoderVideoOutput::EImageCopyResult::NotSupported;
-	}
+	{ return IElectraDecoderVideoOutput::EImageCopyResult::NotSupported; }
 
 #if ELECTRA_HAVE_IMAGEBUFFERS
 	//
@@ -220,7 +198,7 @@ public:
 		}
 		return EElectraDecoderPlatformPixelEncoding::Native;
 	}
-	
+
 	// Returns the n'th image buffer pitch
 	virtual int32 GetBufferPitchByIndex(int32 InBufferIndex) const override
 	{
@@ -419,8 +397,8 @@ FElectraVideoDecoderH265_DX::FElectraVideoDecoderH265_DX(const TMap<FString, FVa
 {
 	ResourceDelegate = InResourceDelegate;
 
-	MaxWidth = (uint32)ElectraDecodersUtil::GetVariantValueSafeU64(InOptions, TEXT("max_width"), 1920);
-	MaxHeight = (uint32)ElectraDecodersUtil::GetVariantValueSafeU64(InOptions, TEXT("max_height"), 1080);
+	MaxWidth = Align((uint32)ElectraDecodersUtil::GetVariantValueSafeU64(InOptions, TEXT("max_width"), 3840), 16u);
+	MaxHeight = Align((uint32)ElectraDecodersUtil::GetVariantValueSafeU64(InOptions, TEXT("max_height"), 2160), 16u);
 
 	MaxOutputBuffers = (uint32)ElectraDecodersUtil::GetVariantValueSafeU64(InOptions, TEXT("max_output_buffers"), 5);
 	MaxOutputBuffers += kElectraDecoderPipelineExtraFrames;
@@ -497,7 +475,7 @@ bool FElectraVideoDecoderH265_DX::ResetToCleanStart()
 	CurrentOutput.Reset();
 	DecodeState = EDecodeState::Decoding;
 	bRequireDiscontinuity = true;
-	
+
 	return !LastError.IsSet();
 }
 
@@ -653,7 +631,7 @@ IElectraDecoder::EDecoderError FElectraVideoDecoderH265_DX::SendEndOfData()
 	{
 		return IElectraDecoder::EDecoderError::EndOfData;
 	}
-	
+
 	// If there is a transform send an end-of-stream and drain message.
 	if (DecoderTransform.IsValid())
 	{
@@ -840,7 +818,7 @@ bool FElectraVideoDecoderH265_DX::InternalDecoderCreate(const TMap<FString, FVar
 	}
 	else
 	{
-		VERIFY_HR(MFSetAttributeSize(MediaType, MF_MT_FRAME_SIZE, 1920, 1088), "Failed to set video decoder input media type resolution", ERRCODE_INTERNAL_FAILED_TO_CREATE_DECODER);
+		VERIFY_HR(MFSetAttributeSize(MediaType, MF_MT_FRAME_SIZE, 3840, 2160), "Failed to set video decoder input media type resolution", ERRCODE_INTERNAL_FAILED_TO_CREATE_DECODER);
 	}
 
 	// Set input media type with decoder
@@ -1106,7 +1084,7 @@ bool FElectraVideoDecoderH265_DX::ConvertDecoderOutput()
 	if (DecoderPlatformHandle->IsSoftware() || (DecoderPlatformHandle->GetDXVersionTimes1000() >= 12000 && DecoderPlatformHandle->GetDXDevice() != nullptr) || DecoderPlatformHandle->GetDXVersionTimes1000() == 0)
 	{
 #if !ALLOW_MFSAMPLE_WITH_DX12
-		if (DecoderPlatformHandle->GetDXVersionTimes1000() < 12000)
+		if (DecoderPlatformHandle->GetDXVersionTimes1000() < 12000 || CVarElectraWindowsH265UseOldOutputPath.GetValueOnAnyThread() != 0)
 #endif
 		{
 			// We are returning a CPU side buffer (software decode OR DX12 render device) - adjust the height so it can be interpreted as a single plane texture
@@ -1175,7 +1153,7 @@ bool FElectraVideoDecoderH265_DX::ConvertDecoderOutput()
 	int32 PlatformDeviceVersion = 0;
 	bool bUseGPUBuffers = false;
 	auto PinnedResourceDelegate = ResourceDelegate.Pin();
-	if (PinnedResourceDelegate.IsValid())
+	if (PinnedResourceDelegate.IsValid() && CVarElectraWindowsH265UseOldOutputPath.GetValueOnAnyThread() == 0)
 	{
 		PinnedResourceDelegate->GetD3DDevice(&PlatformDevice, &PlatformDeviceVersion);
 		bUseGPUBuffers = (PlatformDevice && PlatformDeviceVersion >= 12000);
@@ -1248,7 +1226,7 @@ bool FElectraVideoDecoderH265_DX::ConvertDecoderOutput()
 		//
 		// Post SDK 22621: direct link from MFSample to DX12 resource
 		//
-		
+
 		TRefCountPtr<IMFDXGIBuffer> DXGIBuffer;
 		if ((Result = Buffer->QueryInterface(__uuidof(IMFDXGIBuffer), (void**)DXGIBuffer.GetInitReference())) != S_OK)
 		{
@@ -1277,8 +1255,8 @@ bool FElectraVideoDecoderH265_DX::ConvertDecoderOutput()
 	{
 		//
 		// Fallback output path for pre- / none-DX12 graphics APIs
-		// 
-		
+		//
+
 		// Retain the IMFSample in the output. It is needed later in converting it for display.
 		NewOutput->Pitch = stride;
 		NewOutput->MFSample = DecodedOutputSample;
