ElectraCodecs: Added creation of the AVC decoder configuration record from codec specific data

#rnx

diff --git a/Engine/Plugins/Media/ElectraCodecs/Source/ElectraDecoders/Private/Utils/MPEG/ElectraUtilsMPEGVideo.cpp b/Engine/Plugins/Media/ElectraCodecs/Source/ElectraDecoders/Private/Utils/MPEG/ElectraUtilsMPEGVideo.cpp
--- a/Engine/Plugins/Media/ElectraCodecs/Source/ElectraDecoders/Private/Utils/MPEG/ElectraUtilsMPEGVideo.cpp
+++ b/Engine/Plugins/Media/ElectraCodecs/Source/ElectraDecoders/Private/Utils/MPEG/ElectraUtilsMPEGVideo.cpp
@@ -2,6 +2,7 @@
 
 #include "Utils/MPEG/ElectraUtilsMPEGVideo.h"
 #include "Utils/ElectraBitstreamReader.h"
+#include "Utils/MPEG/ElectraUtilsMPEGVideo_H264.h"
 #include "Utils/MPEG/ElectraUtilsMPEGVideo_H265.h"
 #include "ElectraDecodersUtils.h"
 
@@ -52,6 +53,97 @@ namespace ElectraDecodersUtil
 		}
 
 
+		bool FAVCDecoderConfigurationRecord::CreateFromCodecSpecificData(const TArray<uint8>& InFromCSD)
+		{
+			if (!InFromCSD.IsEmpty())
+			{
+				CodecSpecificData = InFromCSD;
+				TArray<ElectraDecodersUtil::MPEG::H264::FNaluInfo> Nalus;
+				if (!ElectraDecodersUtil::MPEG::H264::ParseBitstreamForNALUs(Nalus, InFromCSD.GetData(), InFromCSD.Num()))
+				{
+					return false;
+				}
+
+				FArray aSPS, aPPS;
+				for(int32 i=0; i<Nalus.Num(); ++i)
+				{
+					if (Nalus[i].Type == 7)
+					{
+						TMap<uint32, ElectraDecodersUtil::MPEG::H264::FSequenceParameterSet> SPSs;
+						if (ElectraDecodersUtil::MPEG::H264::ParseSequenceParameterSet(SPSs, InFromCSD.GetData() + Nalus[i].Offset + Nalus[i].UnitLength, Nalus[i].Size))
+						{
+							ElectraDecodersUtil::MPEG::H264::FSequenceParameterSet sps(SPSs.CreateConstIterator().Value());
+							AVCProfileIndication = sps.profile_idc;
+							ProfileCompatibility = (sps.constraint_set0_flag << 7) | (sps.constraint_set1_flag << 6) | (sps.constraint_set2_flag << 5) |
+												   (sps.constraint_set3_flag << 4) | (sps.constraint_set4_flag << 3) | (sps.constraint_set5_flag << 2);
+							AVCLevelIndication = sps.level_idc;
+							NALUnitLength = 4;
+							ChromaFormat = sps.chroma_format_idc;
+							BitDepthLumaMinus8 = sps.bit_depth_luma_minus8;
+							BitDepthChromaMinus8 = sps.bit_depth_chroma_minus8;
+							bHaveAdditionalProfileIndication = false;
+						}
+						else
+						{
+							return false;
+						}
+						aSPS.NALUs.Emplace_GetRef() = MakeArrayView<const uint8>(InFromCSD.GetData() + Nalus[i].Offset + Nalus[i].UnitLength, Nalus[i].Size);
+					}
+					else if (Nalus[i].Type == 8)
+					{
+						aPPS.NALUs.Emplace_GetRef() = MakeArrayView<const uint8>(InFromCSD.GetData() + Nalus[i].Offset + Nalus[i].UnitLength, Nalus[i].Size);
+					}
+				}
+				if (aSPS.NALUs.IsEmpty() || aPPS.NALUs.IsEmpty())
+				{
+					return false;
+				}
+
+				// Create the raw configuration record now.
+				ConfigurationVersion = 1;
+				FElectraBitstreamWriter wr;
+				wr.PutBits(ConfigurationVersion, 8);
+				wr.PutBits(AVCProfileIndication, 8);
+				wr.PutBits(ProfileCompatibility, 8);
+				wr.PutBits(AVCLevelIndication, 8);
+				wr.PutBits(63U, 6);
+				wr.PutBits(NALUnitLength-1, 2);
+				wr.PutBits(7U, 3);
+				wr.PutBits(static_cast<uint32>(aSPS.NALUs.Num()), 5);
+				for(int32 j=0; j<aSPS.NALUs.Num(); ++j)
+				{
+					const TArray<uint8>& d(aSPS.NALUs[j]);
+					wr.PutBits(static_cast<uint32>(d.Num()), 16);
+					for(int32 k=0; k<d.Num(); ++k)
+					{
+						wr.PutBits(d[k], 8);
+					}
+				}
+				for(int32 j=0; j<aPPS.NALUs.Num(); ++j)
+				{
+					const TArray<uint8>& d(aPPS.NALUs[j]);
+					wr.PutBits(static_cast<uint32>(d.Num()), 16);
+					for(int32 k=0; k<d.Num(); ++k)
+					{
+						wr.PutBits(d[k], 8);
+					}
+				}
+				if (AVCProfileIndication != 66 && AVCProfileIndication != 77 && AVCProfileIndication != 88)
+				{
+					wr.PutBits(63U, 6);
+					wr.PutBits(ChromaFormat, 2);
+					wr.PutBits(31U, 5);
+					wr.PutBits(BitDepthLumaMinus8, 3);
+					wr.PutBits(31U, 5);
+					wr.PutBits(BitDepthChromaMinus8, 3);
+					wr.PutBits(0U, 8);
+				}
+				wr.GetArray(RawData);
+				return true;
+			}
+			return false;
+		}
+
 		bool FAVCDecoderConfigurationRecord::Parse()
 		{
 			CodecSpecificData.Empty();
diff --git a/Engine/Plugins/Media/ElectraCodecs/Source/ElectraDecoders/Public/Utils/MPEG/ElectraUtilsMPEGVideo.h b/Engine/Plugins/Media/ElectraCodecs/Source/ElectraDecoders/Public/Utils/MPEG/ElectraUtilsMPEGVideo.h
--- a/Engine/Plugins/Media/ElectraCodecs/Source/ElectraDecoders/Public/Utils/MPEG/ElectraUtilsMPEGVideo.h
+++ b/Engine/Plugins/Media/ElectraCodecs/Source/ElectraDecoders/Public/Utils/MPEG/ElectraUtilsMPEGVideo.h
@@ -153,6 +153,7 @@ namespace ElectraDecodersUtil
 		public:
 			void SetRawData(const void* Data, int64 Size);
 			const TArray<uint8>& GetRawData() const;
+			bool CreateFromCodecSpecificData(const TArray<uint8>& InFromCSD);
 
 			bool Parse();
 
@@ -171,6 +172,11 @@ namespace ElectraDecodersUtil
 			{ return AVCLevelIndication; }
 
 		private:
+			struct FArray
+			{
+				TArray<TArray<uint8>> NALUs;
+			};
+
 			TArray<uint8>													RawData;
 			TArray<uint8>													CodecSpecificData;
 			TArray<uint8>													CodecSpecificDataSPSOnly;
diff --git a/Engine/Plugins/Media/ElectraPlayer/Source/ElectraPlayerRuntime/Private/Runtime/Player/StreamSegmentReaderCommon.cpp b/Engine/Plugins/Media/ElectraPlayer/Source/ElectraPlayerRuntime/Private/Runtime/Player/StreamSegmentReaderCommon.cpp
--- a/Engine/Plugins/Media/ElectraPlayer/Source/ElectraPlayerRuntime/Private/Runtime/Player/StreamSegmentReaderCommon.cpp
+++ b/Engine/Plugins/Media/ElectraPlayer/Source/ElectraPlayerRuntime/Private/Runtime/Player/StreamSegmentReaderCommon.cpp
@@ -2775,7 +2775,21 @@ FStreamSegmentReaderCommon::FStreamHandler::EHandleResult FStreamSegmentReaderCo
 							CurrentlyActiveTrackData->DefaultDurationFromCSD.SetFromND(1, 30);
 						}
 						// Some decoders need the ISO/IEC 14496-15 decoder configuration record, which we need to construct.
-						if (CurrentlyActiveTrackData->CSD->ParsedInfo.GetCodec() == FStreamCodecInformation::ECodec::H265 &&
+						if (CurrentlyActiveTrackData->CSD->ParsedInfo.GetCodec() == FStreamCodecInformation::ECodec::H264 &&
+							CurrentlyActiveTrackData->CSD->RawCSD.IsEmpty())
+						{
+							ElectraDecodersUtil::MPEG::FAVCDecoderConfigurationRecord dcr;
+							if (dcr.CreateFromCodecSpecificData(CurrentlyActiveTrackData->CSD->CodecSpecificData))
+							{
+								CurrentlyActiveTrackData->CSD->RawCSD = dcr.GetRawData();
+							}
+							else
+							{
+								SetError(FString::Printf(TEXT("Failed to create the H.264 decoder configuration record from the inband CSD")), INTERNAL_SEG_ERROR_BAD_MEDIA_SEGMENT);
+								return EHandleResult::Failed;
+							}
+						}
+						else if (CurrentlyActiveTrackData->CSD->ParsedInfo.GetCodec() == FStreamCodecInformation::ECodec::H265 &&
 							CurrentlyActiveTrackData->CSD->RawCSD.IsEmpty())
 						{
 							ElectraDecodersUtil::MPEG::FHEVCDecoderConfigurationRecord dcr;
