ElectraCodecs: added convenience SPS parsing methods to parse only a single one

#rnx

diff --git a/Engine/Plugins/Media/ElectraCodecs/Source/ElectraDecoders/Private/Utils/MPEG/ElectraUtilsMPEGVideo_H264.cpp b/Engine/Plugins/Media/ElectraCodecs/Source/ElectraDecoders/Private/Utils/MPEG/ElectraUtilsMPEGVideo_H264.cpp
--- a/Engine/Plugins/Media/ElectraCodecs/Source/ElectraDecoders/Private/Utils/MPEG/ElectraUtilsMPEGVideo_H264.cpp
+++ b/Engine/Plugins/Media/ElectraCodecs/Source/ElectraDecoders/Private/Utils/MPEG/ElectraUtilsMPEGVideo_H264.cpp
@@ -375,6 +375,16 @@ namespace ElectraDecodersUtil
 			}
 
 
+			bool ParseSequenceParameterSet(FSequenceParameterSet& OutSequenceParameterSet, const uint8* InBitstream, uint64 InBitstreamLenInBytes)
+			{
+				TMap<uint32, FSequenceParameterSet> spss;
+				if (!ParseSequenceParameterSet(spss, InBitstream, InBitstreamLenInBytes))
+				{
+					return false;
+				}
+				OutSequenceParameterSet = spss.CreateConstIterator().Value();
+				return true;
+			}
 
 			bool ParseSequenceParameterSet(TMap<uint32, FSequenceParameterSet>& InOutSequenceParameterSets, const uint8* InBitstream, uint64 InBitstreamLenInBytes)
 			{
diff --git a/Engine/Plugins/Media/ElectraCodecs/Source/ElectraDecoders/Private/Utils/MPEG/ElectraUtilsMPEGVideo_H265.cpp b/Engine/Plugins/Media/ElectraCodecs/Source/ElectraDecoders/Private/Utils/MPEG/ElectraUtilsMPEGVideo_H265.cpp
--- a/Engine/Plugins/Media/ElectraCodecs/Source/ElectraDecoders/Private/Utils/MPEG/ElectraUtilsMPEGVideo_H265.cpp
+++ b/Engine/Plugins/Media/ElectraCodecs/Source/ElectraDecoders/Private/Utils/MPEG/ElectraUtilsMPEGVideo_H265.cpp
@@ -1017,6 +1017,17 @@ namespace ElectraDecodersUtil
 				#undef RANGE_CHECK_FAILURE
 			}
 
+			bool ParseSequenceParameterSet(FSequenceParameterSet& OutSequenceParameterSet, const uint8* InBitstream, uint64 InBitstreamLenInBytes)
+			{
+				TMap<uint32, FSequenceParameterSet> spss;
+				if (!ParseSequenceParameterSet(spss, InBitstream, InBitstreamLenInBytes))
+				{
+					return false;
+				}
+				OutSequenceParameterSet = spss.CreateConstIterator().Value();
+				return true;
+			}
+
 			bool ParseSequenceParameterSet(TMap<uint32, FSequenceParameterSet>& InOutSequenceParameterSets, const uint8* InBitstream, uint64 InBitstreamLenInBytes)
 			{
 				#define RANGE_CHECK_FAILURE(expr, elem)													\
diff --git a/Engine/Plugins/Media/ElectraCodecs/Source/ElectraDecoders/Public/Utils/MPEG/ElectraUtilsMPEGVideo_H264.h b/Engine/Plugins/Media/ElectraCodecs/Source/ElectraDecoders/Public/Utils/MPEG/ElectraUtilsMPEGVideo_H264.h
--- a/Engine/Plugins/Media/ElectraCodecs/Source/ElectraDecoders/Public/Utils/MPEG/ElectraUtilsMPEGVideo_H264.h
+++ b/Engine/Plugins/Media/ElectraCodecs/Source/ElectraDecoders/Public/Utils/MPEG/ElectraUtilsMPEGVideo_H264.h
@@ -576,6 +576,7 @@ namespace ElectraDecodersUtil
 			};
 
 
+			bool ELECTRADECODERS_API ParseSequenceParameterSet(FSequenceParameterSet& OutSequenceParameterSet, const uint8* InBitstream, uint64 InBitstreamLenInBytes);
 			bool ELECTRADECODERS_API ParseSequenceParameterSet(TMap<uint32, FSequenceParameterSet>& InOutSequenceParameterSets, const uint8* InBitstream, uint64 InBitstreamLenInBytes);
 			bool ELECTRADECODERS_API ParsePictureParameterSet(TMap<uint32, FPictureParameterSet>& InOutPictureParameterSets, const TMap<uint32, FSequenceParameterSet>& InSequenceParameterSets, const uint8* InBitstream, uint64 InBitstreamLenInBytes);
 			bool ELECTRADECODERS_API ParseSliceHeader(TUniquePtr<FRBSP>& OutRBSP, FBitstreamReader& OutRBSPReader, FSliceHeader& OutSlice, const TMap<uint32, FSequenceParameterSet>& InSequenceParameterSets, const TMap<uint32, FPictureParameterSet>& InPictureParameterSets, const uint8* InBitstream, uint64 InBitstreamLenInBytes);
diff --git a/Engine/Plugins/Media/ElectraCodecs/Source/ElectraDecoders/Public/Utils/MPEG/ElectraUtilsMPEGVideo_H265.h b/Engine/Plugins/Media/ElectraCodecs/Source/ElectraDecoders/Public/Utils/MPEG/ElectraUtilsMPEGVideo_H265.h
--- a/Engine/Plugins/Media/ElectraCodecs/Source/ElectraDecoders/Public/Utils/MPEG/ElectraUtilsMPEGVideo_H265.h
+++ b/Engine/Plugins/Media/ElectraCodecs/Source/ElectraDecoders/Public/Utils/MPEG/ElectraUtilsMPEGVideo_H265.h
@@ -1357,6 +1357,7 @@ namespace ElectraDecodersUtil
 
 			bool ELECTRADECODERS_API ParseVideoParameterSet(TMap<uint32, FVideoParameterSet>& InOutVideoParameterSets, const uint8* InBitstream, uint64 InBitstreamLenInBytes);
 			bool ELECTRADECODERS_API ParseSequenceParameterSet(TMap<uint32, FSequenceParameterSet>& InOutSequenceParameterSets, const uint8* InBitstream, uint64 InBitstreamLenInBytes);
+			bool ELECTRADECODERS_API ParseSequenceParameterSet(FSequenceParameterSet& OutSequenceParameterSet, const uint8* InBitstream, uint64 InBitstreamLenInBytes);
 			bool ELECTRADECODERS_API ParsePictureParameterSet(TMap<uint32, FPictureParameterSet>& InOutPictureParameterSets, const TMap<uint32, FSequenceParameterSet>& InSequenceParameterSets, const uint8* InBitstream, uint64 InBitstreamLenInBytes);
 			bool ELECTRADECODERS_API ParseSliceHeader(TUniquePtr<FRBSP>& OutRBSP, FBitstreamReader& OutRBSPReader, FSliceSegmentHeader& OutSlice, const TMap<uint32, FVideoParameterSet>& InVideoParameterSets, const TMap<uint32, FSequenceParameterSet>& InSequenceParameterSets, const TMap<uint32, FPictureParameterSet>& InPictureParameterSets, const uint8* InBitstream, uint64 InBitstreamLenInBytes);
 
