ElectraDecoders: Resetting the H.264 decoder MFT when the stream DPB size changes to avoid an issue with the decoder not proceeding after an output format change

#rnx

diff --git a/Engine/Plugins/Media/ElectraCodecs/Source/ElectraDecoders/Private/DX/VideoDecoderH264_DX.cpp b/Engine/Plugins/Media/ElectraCodecs/Source/ElectraDecoders/Private/DX/VideoDecoderH264_DX.cpp
--- a/Engine/Plugins/Media/ElectraCodecs/Source/ElectraDecoders/Private/DX/VideoDecoderH264_DX.cpp
+++ b/Engine/Plugins/Media/ElectraCodecs/Source/ElectraDecoders/Private/DX/VideoDecoderH264_DX.cpp
@@ -6,6 +6,7 @@
 #include "DX/DecoderErrors_DX.h"
 #include "ElectraDecodersUtils.h"
 #include "Utils/MPEG/ElectraUtilsMPEGVideo.h"
+#include "Utils/MPEG/ElectraUtilsMPEGVideo_H264.h"
 
 #include "IElectraDecoderFeaturesAndOptions.h"
 #include "IElectraDecoderOutputVideo.h"
@@ -101,7 +102,7 @@ public:
 	void GetExtraValues(TMap<FString, FVariant>& OutExtraValues) const override
 	{ OutExtraValues = ExtraValues; }
 	void* GetPlatformOutputHandle(EElectraDecoderPlatformOutputHandleType InTypeOfHandle) const override
-	{ 
+	{
 		switch(InTypeOfHandle)
 		{
 #if ELECTRA_HAVE_IMAGEBUFFERS
@@ -116,11 +117,11 @@ public:
 				return nullptr;
 #endif
 			case EElectraDecoderPlatformOutputHandleType::MFSample:
-				return MFSample.GetReference(); 
+				return MFSample.GetReference();
 			case EElectraDecoderPlatformOutputHandleType::DXDevice:
-				return Ctx.DxDevice.GetReference(); 
+				return Ctx.DxDevice.GetReference();
 			case EElectraDecoderPlatformOutputHandleType::DXDeviceContext:
-				return Ctx.DxDeviceContext.GetReference(); 
+				return Ctx.DxDeviceContext.GetReference();
 			default:
 				return nullptr;
 		}
@@ -198,7 +199,7 @@ public:
 		}
 		return EElectraDecoderPlatformPixelEncoding::Native;
 	}
-	
+
 	// Returns the n'th image buffer pitch
 	virtual int32 GetBufferPitchByIndex(int32 InBufferIndex) const override
 	{
@@ -339,7 +340,8 @@ private:
 
 	void PostError(HRESULT ApiReturnValue, FString Message, int32 Code);
 
-	TSharedPtr<ElectraDecodersUtil::MPEG::FISO14496_10_seq_parameter_set_data, ESPMode::ThreadSafe> GetSPSFromOptions(const TMap<FString, FVariant>& InOptions);
+	TSharedPtr<ElectraDecodersUtil::MPEG::H264::FSequenceParameterSet, ESPMode::ThreadSafe> GetSPSFromOptions(const TMap<FString, FVariant>& InOptions);
+	TSharedPtr<ElectraDecodersUtil::MPEG::H264::FSequenceParameterSet, ESPMode::ThreadSafe> GetSPSFromCSD(const TArray<uint8>& InCSD);
 
 	bool InternalDecoderCreate(const TMap<FString, FVariant>& InAdditionalOptions);
 	void InternalDecoderDestroy();
@@ -359,6 +361,7 @@ private:
 	TUniquePtr<FDecoderOutputBuffer> CurrentDecoderOutputBuffer;
 	TArray<FDecoderInput> InDecoderInput;
 	TSharedPtr<FElectraVideoDecoderOutputH264_DX, ESPMode::ThreadSafe> CurrentOutput;
+	TSharedPtr<ElectraDecodersUtil::MPEG::H264::FSequenceParameterSet, ESPMode::ThreadSafe> CurrentSPS;
 
 	EDecodeState DecodeState = EDecodeState::Decoding;
 	bool bRequireDiscontinuity = false;
@@ -413,19 +416,18 @@ FElectraVideoDecoderH264_DX::~FElectraVideoDecoderH264_DX()
 	Close();
 }
 
-TSharedPtr<ElectraDecodersUtil::MPEG::FISO14496_10_seq_parameter_set_data, ESPMode::ThreadSafe> FElectraVideoDecoderH264_DX::GetSPSFromOptions(const TMap<FString, FVariant>& InOptions)
+TSharedPtr<ElectraDecodersUtil::MPEG::H264::FSequenceParameterSet, ESPMode::ThreadSafe> FElectraVideoDecoderH264_DX::GetSPSFromCSD(const TArray<uint8>& InCSD)
 {
-	TArray<uint8> SidebandData = ElectraDecodersUtil::GetVariantValueUInt8Array(InOptions, TEXT("csd"));
-	if (SidebandData.Num())
+	if (InCSD.Num())
 	{
-		TArray<ElectraDecodersUtil::MPEG::FNaluInfo> NALUs;
-		ElectraDecodersUtil::MPEG::ParseBitstreamForNALUs(NALUs, SidebandData.GetData(), SidebandData.Num());
+		TArray<ElectraDecodersUtil::MPEG::H264::FNaluInfo> NALUs;
+		ElectraDecodersUtil::MPEG::H264::ParseBitstreamForNALUs(NALUs, InCSD.GetData(), InCSD.Num());
 		for(int32 i=0; i<NALUs.Num(); ++i)
 		{
-			if ((NALUs[i].Type & 0x1f) == 7)
+			if (NALUs[i].Type == 7)
 			{
-				TSharedPtr<ElectraDecodersUtil::MPEG::FISO14496_10_seq_parameter_set_data, ESPMode::ThreadSafe> NewSPS = MakeShared<ElectraDecodersUtil::MPEG::FISO14496_10_seq_parameter_set_data, ESPMode::ThreadSafe>();
-				if (ElectraDecodersUtil::MPEG::ParseH264SPS(*NewSPS, ElectraDecodersUtil::AdvancePointer(SidebandData.GetData(), NALUs[i].Offset + NALUs[i].UnitLength), NALUs[i].Size))
+				TSharedPtr<ElectraDecodersUtil::MPEG::H264::FSequenceParameterSet, ESPMode::ThreadSafe> NewSPS = MakeShared<ElectraDecodersUtil::MPEG::H264::FSequenceParameterSet, ESPMode::ThreadSafe>();
+				if (ElectraDecodersUtil::MPEG::H264::ParseSequenceParameterSet(*NewSPS, ElectraDecodersUtil::AdvancePointer(InCSD.GetData(), NALUs[i].Offset + NALUs[i].UnitLength), NALUs[i].Size))
 				{
 					return NewSPS;
 				}
@@ -435,6 +437,20 @@ TSharedPtr<ElectraDecodersUtil::MPEG::FISO14496_10_seq_parameter_set_data, ESPMo
 				}
 			}
 		}
+	}
+	return nullptr;
+}
+
+TSharedPtr<ElectraDecodersUtil::MPEG::H264::FSequenceParameterSet, ESPMode::ThreadSafe> FElectraVideoDecoderH264_DX::GetSPSFromOptions(const TMap<FString, FVariant>& InOptions)
+{
+	TArray<uint8> SidebandData = ElectraDecodersUtil::GetVariantValueUInt8Array(InOptions, TEXT("csd"));
+	if (SidebandData.Num())
+	{
+		auto NewSPS = GetSPSFromCSD(SidebandData);
+		if (NewSPS.IsValid())
+		{
+			return NewSPS;
+		}
 		LastError.Code = ERRCODE_INTERNAL_FAILED_TO_PARSE_CSD;
 		LastError.Message = TEXT("Failed to parse codec specific data");
 	}
@@ -472,9 +488,10 @@ bool FElectraVideoDecoderH264_DX::ResetToCleanStart()
 	InternalDecoderDestroy();
 	InDecoderInput.Empty();
 	CurrentOutput.Reset();
+	CurrentSPS.Reset();
 	DecodeState = EDecodeState::Decoding;
 	bRequireDiscontinuity = true;
-	
+
 	return !LastError.IsSet();
 }
 
@@ -492,6 +509,17 @@ IElectraDecoder::ECSDCompatibility FElectraVideoDecoderH264_DX::IsCompatibleWith
 	{
 		return IElectraDecoder::ECSDCompatibility::Compatible;
 	}
+
+	// We need to check for a change in DPB size. If that happens the MFT needs to be reset, otherwise it may get stuck
+	// in an infinite ProcessOutput() returning MF_E_TRANSFORM_STREAM_CHANGE all the time.
+	TSharedPtr<ElectraDecodersUtil::MPEG::H264::FSequenceParameterSet, ESPMode::ThreadSafe> NewSPS(GetSPSFromOptions(CSDAndAdditionalOptions));
+	if (NewSPS.IsValid() && CurrentSPS.IsValid())
+	{
+		if (NewSPS->GetDPBSize() != CurrentSPS->GetDPBSize())
+		{
+			return IElectraDecoder::ECSDCompatibility::DrainAndReset;
+		}
+	}
 	return IElectraDecoder::ECSDCompatibility::Drain;
 }
 
@@ -630,7 +658,7 @@ IElectraDecoder::EDecoderError FElectraVideoDecoderH264_DX::SendEndOfData()
 	{
 		return IElectraDecoder::EDecoderError::EndOfData;
 	}
-	
+
 	// If there is a transform send an end-of-stream and drain message.
 	if (DecoderTransform.IsValid())
 	{
@@ -668,6 +696,7 @@ IElectraDecoder::EDecoderError FElectraVideoDecoderH264_DX::Flush()
 		bRequireDiscontinuity = true;
 		InDecoderInput.Empty();
 		CurrentOutput.Reset();
+		CurrentSPS.Reset();
 	}
 	return IElectraDecoder::EDecoderError::None;
 }
@@ -716,6 +745,7 @@ IElectraDecoder::EOutputStatus FElectraVideoDecoderH264_DX::HaveOutput()
 					CurrentDecoderOutputBuffer.Reset();
 					bRequireDiscontinuity = true;
 					InDecoderInput.Empty();
+					CurrentSPS.Reset();
 					return IElectraDecoder::EOutputStatus::EndOfData;
 				}
 				else
@@ -925,6 +955,7 @@ bool FElectraVideoDecoderH264_DX::CreateInputSample(TRefCountPtr<IMFSample>& Inp
 	if ((InInputAccessUnit.Flags & EElectraDecoderFlags::IsSyncSample) != EElectraDecoderFlags::None)
 	{
 		CSD = ElectraDecodersUtil::GetVariantValueUInt8Array(InAdditionalOptions, TEXT("csd"));
+		CurrentSPS = GetSPSFromCSD(CSD);
 	}
 	InputDataSize = CSD.Num() + InInputAccessUnit.DataSize;
 	VERIFY_HR(MFCreateSample(InputSample.GetInitReference()), TEXT("Failed to create video decoder input sample"), ERRCODE_INTERNAL_FAILED_TO_CREATE_INPUT_SAMPLE);
@@ -1190,7 +1221,7 @@ bool FElectraVideoDecoderH264_DX::ConvertDecoderOutput()
 		//
 		// Post SDK 22621: direct link from MFSample to DX12 resource
 		//
-		
+
 		TRefCountPtr<IMFDXGIBuffer> DXGIBuffer;
 		if ((Result = Buffer->QueryInterface(__uuidof(IMFDXGIBuffer), (void**)DXGIBuffer.GetInitReference())) != S_OK)
 		{
@@ -1219,8 +1250,8 @@ bool FElectraVideoDecoderH264_DX::ConvertDecoderOutput()
 	{
 		//
 		// Fallback output path for pre- / none-DX12 graphics APIs
-		// 
-		
+		//
+
 		// Retain the IMFSample in the output. It is needed later in converting it for display.
 		NewOutput->Pitch = stride;
 		NewOutput->MFSample = DecodedOutputSample;
