ElectraDecoders: Also resetting the H.264 decoder MFT on a seek where the previous DPB setting is not available

#rnx

diff --git a/Engine/Plugins/Media/ElectraCodecs/Source/ElectraDecoders/Private/DX/VideoDecoderH264_DX.cpp b/Engine/Plugins/Media/ElectraCodecs/Source/ElectraDecoders/Private/DX/VideoDecoderH264_DX.cpp
--- a/Engine/Plugins/Media/ElectraCodecs/Source/ElectraDecoders/Private/DX/VideoDecoderH264_DX.cpp
+++ b/Engine/Plugins/Media/ElectraCodecs/Source/ElectraDecoders/Private/DX/VideoDecoderH264_DX.cpp
@@ -513,12 +513,9 @@ IElectraDecoder::ECSDCompatibility FElectraVideoDecoderH264_DX::IsCompatibleWith
 	// We need to check for a change in DPB size. If that happens the MFT needs to be reset, otherwise it may get stuck
 	// in an infinite ProcessOutput() returning MF_E_TRANSFORM_STREAM_CHANGE all the time.
 	TSharedPtr<ElectraDecodersUtil::MPEG::H264::FSequenceParameterSet, ESPMode::ThreadSafe> NewSPS(GetSPSFromOptions(CSDAndAdditionalOptions));
-	if (NewSPS.IsValid() && CurrentSPS.IsValid())
+	if (!CurrentSPS.IsValid() || (NewSPS.IsValid() && CurrentSPS.IsValid() && NewSPS->GetDPBSize() != CurrentSPS->GetDPBSize()))
 	{
-		if (NewSPS->GetDPBSize() != CurrentSPS->GetDPBSize())
-		{
-			return IElectraDecoder::ECSDCompatibility::DrainAndReset;
-		}
+		return IElectraDecoder::ECSDCompatibility::DrainAndReset;
 	}
 	return IElectraDecoder::ECSDCompatibility::Drain;
 }
