ElectraCodecs: Added missing field in generating the H.264 DCR from CSD

#rnx

diff --git a/Engine/Plugins/Media/ElectraCodecs/Source/ElectraDecoders/Private/Utils/MPEG/ElectraUtilsMPEGVideo.cpp b/Engine/Plugins/Media/ElectraCodecs/Source/ElectraDecoders/Private/Utils/MPEG/ElectraUtilsMPEGVideo.cpp
--- a/Engine/Plugins/Media/ElectraCodecs/Source/ElectraDecoders/Private/Utils/MPEG/ElectraUtilsMPEGVideo.cpp
+++ b/Engine/Plugins/Media/ElectraCodecs/Source/ElectraDecoders/Private/Utils/MPEG/ElectraUtilsMPEGVideo.cpp
@@ -119,6 +119,7 @@ namespace ElectraDecodersUtil
 						wr.PutBits(d[k], 8);
 					}
 				}
+				wr.PutBits(static_cast<uint32>(aPPS.NALUs.Num()), 8);
 				for(int32 j=0; j<aPPS.NALUs.Num(); ++j)
 				{
 					const TArray<uint8>& d(aPPS.NALUs[j]);
