ElectraCodecs: Fixed incorrect parsing of legacy avc OTI "avc1.xx.yy"

#rnx

diff --git a/Engine/Plugins/Media/ElectraCodecs/Source/ElectraDecoders/Private/ElectraDecodersUtils.cpp b/Engine/Plugins/Media/ElectraCodecs/Source/ElectraDecoders/Private/ElectraDecodersUtils.cpp
--- a/Engine/Plugins/Media/ElectraCodecs/Source/ElectraDecoders/Private/ElectraDecodersUtils.cpp
+++ b/Engine/Plugins/Media/ElectraCodecs/Source/ElectraDecoders/Private/ElectraDecodersUtils.cpp
@@ -9,10 +9,10 @@ namespace ElectraDecodersUtil
 namespace
 {
 	inline void LexFromStringHex(int32& OutValue, const TCHAR* Buffer)
-	{ 
+	{
 		OutValue = FCString::Strtoi(Buffer, nullptr, 16);
 	}
-	
+
 	inline void LexFromStringHexU64(uint64& OutValue, const TCHAR* Buffer)
 	{
 		OutValue = FCString::Strtoui64(Buffer, nullptr, 16);
@@ -72,10 +72,12 @@ bool ParseCodecH264(FMimeTypeVideoCodecInfo& OutInfo, const FString& InCodecForm
 			// Profile and level follow?
 			if (InCodecFormat.Len() > 5 && InCodecFormat[4] == TCHAR('.'))
 			{
+				int32 DotPos;
+				InCodecFormat.FindLastChar(TCHAR('.'), DotPos);
 				FString Temp;
 				int32 TempValue;
-				// We recognize the expected format avcC.xxyyzz and for legacy reasons also avcC.xx.zz
-				if (InCodecFormat.Len() == 11)
+				// We recognize the expected format avcC.xxyyzz and for legacy reasons also avcC.xxx.zz
+				if (InCodecFormat.Len() == 11 && DotPos == 4)
 				{
 					Temp = InCodecFormat.Mid(5, 2);
 					LexFromStringHex(TempValue, *Temp);
@@ -87,13 +89,13 @@ bool ParseCodecH264(FMimeTypeVideoCodecInfo& OutInfo, const FString& InCodecForm
 					LexFromStringHex(TempValue, *Temp);
 					OutInfo.Level = TempValue;
 				}
-				else if (InCodecFormat.Len() == 10 && InCodecFormat[7] == TCHAR('.'))
+				else if (DotPos != INDEX_NONE)
 				{
-					Temp = InCodecFormat.Mid(5, 2);
-					LexFromStringHex(TempValue, *Temp);
+					Temp = InCodecFormat.Mid(5, DotPos-5);
+					LexFromString(TempValue, *Temp);
 					OutInfo.Profile = TempValue;
-					Temp = InCodecFormat.Mid(8, 2);
-					LexFromStringHex(TempValue, *Temp);
+					Temp = InCodecFormat.Mid(DotPos+1);
+					LexFromString(TempValue, *Temp);
 					OutInfo.Level = TempValue;
 				}
 				else
