ElectraCodecs: Fixed copy-pasted wrong factory class name

#rnx

diff --git a/Engine/Plugins/Media/ElectraCodecs/Source/ElectraDecoders/Private/ElectraMediaMP3Decoder.cpp b/Engine/Plugins/Media/ElectraCodecs/Source/ElectraDecoders/Private/ElectraMediaMP3Decoder.cpp
--- a/Engine/Plugins/Media/ElectraCodecs/Source/ElectraDecoders/Private/ElectraMediaMP3Decoder.cpp
+++ b/Engine/Plugins/Media/ElectraCodecs/Source/ElectraDecoders/Private/ElectraMediaMP3Decoder.cpp
@@ -218,10 +218,10 @@ private:
 /*********************************************************************************************************************/
 /*********************************************************************************************************************/
 
-class FElectraCommonAudioFlacDecoderFactory : public IElectraCodecFactory, public IElectraCodecModularFeature, public TSharedFromThis<FElectraCommonAudioFlacDecoderFactory, ESPMode::ThreadSafe>
+class FElectraCommonAudioMP3DecoderFactory : public IElectraCodecFactory, public IElectraCodecModularFeature, public TSharedFromThis<FElectraCommonAudioMP3DecoderFactory, ESPMode::ThreadSafe>
 {
 public:
-	virtual ~FElectraCommonAudioFlacDecoderFactory()
+	virtual ~FElectraCommonAudioMP3DecoderFactory()
 	{ }
 
 	void GetListOfFactories(TArray<TWeakPtr<IElectraCodecFactory, ESPMode::ThreadSafe>>& OutCodecFactories) override
@@ -249,11 +249,11 @@ public:
 		return MakeShared<FElectraMP3Decoder, ESPMode::ThreadSafe>(InOptions, InResourceDelegate);
 	}
 
-	static TSharedPtr<FElectraCommonAudioFlacDecoderFactory, ESPMode::ThreadSafe> Self;
+	static TSharedPtr<FElectraCommonAudioMP3DecoderFactory, ESPMode::ThreadSafe> Self;
 	static TArray<FString> Permitted4CCs;
 };
-TSharedPtr<FElectraCommonAudioFlacDecoderFactory, ESPMode::ThreadSafe> FElectraCommonAudioFlacDecoderFactory::Self;
-TArray<FString> FElectraCommonAudioFlacDecoderFactory::Permitted4CCs = { TEXT("mp4a.6b"), TEXT("mp4a.40.34"), TEXT(".mp3") };
+TSharedPtr<FElectraCommonAudioMP3DecoderFactory, ESPMode::ThreadSafe> FElectraCommonAudioMP3DecoderFactory::Self;
+TArray<FString> FElectraCommonAudioMP3DecoderFactory::Permitted4CCs = { TEXT("mp4a.6b"), TEXT("mp4a.40.34"), TEXT(".mp3") };
 
 /*********************************************************************************************************************/
 /*********************************************************************************************************************/
@@ -266,24 +266,24 @@ void FElectraMediaMP3Decoder::Startup()
 	FModuleManager::Get().LoadModule(TEXT("ElectraCodecFactory"));
 
 	// Create an instance of the factory, which is also the modular feature.
-	check(!FElectraCommonAudioFlacDecoderFactory::Self.IsValid());
-	FElectraCommonAudioFlacDecoderFactory::Self = MakeShared<FElectraCommonAudioFlacDecoderFactory, ESPMode::ThreadSafe>();
+	check(!FElectraCommonAudioMP3DecoderFactory::Self.IsValid());
+	FElectraCommonAudioMP3DecoderFactory::Self = MakeShared<FElectraCommonAudioMP3DecoderFactory, ESPMode::ThreadSafe>();
 	// Register as modular feature.
-	IModularFeatures::Get().RegisterModularFeature(IElectraCodecFactoryModule::GetModularFeatureName(), FElectraCommonAudioFlacDecoderFactory::Self.Get());
+	IModularFeatures::Get().RegisterModularFeature(IElectraCodecFactoryModule::GetModularFeatureName(), FElectraCommonAudioMP3DecoderFactory::Self.Get());
 #endif
 }
 
 void FElectraMediaMP3Decoder::Shutdown()
 {
 #if WITH_DR_LIBS_MP3
-	IModularFeatures::Get().UnregisterModularFeature(IElectraCodecFactoryModule::GetModularFeatureName(), FElectraCommonAudioFlacDecoderFactory::Self.Get());
-	FElectraCommonAudioFlacDecoderFactory::Self.Reset();
+	IModularFeatures::Get().UnregisterModularFeature(IElectraCodecFactoryModule::GetModularFeatureName(), FElectraCommonAudioMP3DecoderFactory::Self.Get());
+	FElectraCommonAudioMP3DecoderFactory::Self.Reset();
 #endif
 }
 
 TSharedPtr<IElectraCodecFactory, ESPMode::ThreadSafe> FElectraMediaMP3Decoder::CreateFactory()
 {
-	return MakeShared<FElectraCommonAudioFlacDecoderFactory, ESPMode::ThreadSafe>();;
+	return MakeShared<FElectraCommonAudioMP3DecoderFactory, ESPMode::ThreadSafe>();;
 }
 
 /*********************************************************************************************************************/
