Electra: Updated handling of decoder confiuration records

#rnx

diff --git a/Engine/Plugins/Media/ElectraCodecs/Source/ElectraDecoders/Private/ElectraDecodersUtils.cpp b/Engine/Plugins/Media/ElectraCodecs/Source/ElectraDecoders/Private/ElectraDecodersUtils.cpp
--- a/Engine/Plugins/Media/ElectraCodecs/Source/ElectraDecoders/Private/ElectraDecodersUtils.cpp
+++ b/Engine/Plugins/Media/ElectraCodecs/Source/ElectraDecoders/Private/ElectraDecodersUtils.cpp
@@ -84,7 +84,7 @@ bool ParseCodecH264(FMimeTypeVideoCodecInfo& OutInfo, const FString& InCodecForm
 					OutInfo.Profile = TempValue;
 					Temp = InCodecFormat.Mid(7, 2);
 					LexFromStringHex(TempValue, *Temp);
-					//ProfileConstraints = TempValue;
+					OutInfo.Constraints = TempValue;
 					Temp = InCodecFormat.Mid(9, 2);
 					LexFromStringHex(TempValue, *Temp);
 					OutInfo.Level = TempValue;
@@ -465,7 +465,6 @@ double GetVariantValueSafeDouble(const TMap<FString, FVariant>& InFromMap, const
 	return V;
 }
 
-
 TArray<uint8> GetVariantValueUInt8Array(const TMap<FString, FVariant>& InFromMap, const FString& InName)
 {
 	const FVariant* Var = InFromMap.Find(InName);
diff --git a/Engine/Plugins/Media/ElectraCodecs/Source/ElectraDecoders/Private/Utils/MPEG/ElectraUtilsMPEGAudio.cpp b/Engine/Plugins/Media/ElectraCodecs/Source/ElectraDecoders/Private/Utils/MPEG/ElectraUtilsMPEGAudio.cpp
--- a/Engine/Plugins/Media/ElectraCodecs/Source/ElectraDecoders/Private/Utils/MPEG/ElectraUtilsMPEGAudio.cpp
+++ b/Engine/Plugins/Media/ElectraCodecs/Source/ElectraDecoders/Private/Utils/MPEG/ElectraUtilsMPEGAudio.cpp
@@ -7,6 +7,123 @@ namespace ElectraDecodersUtil
 {
 	namespace MPEG
 	{
+
+
+		bool FESDescriptor::Parse(const TConstArrayView<uint8>& InESDS)
+		{
+			class FMP4BitReader : public FElectraBitstreamReader
+			{
+			public:
+				FMP4BitReader(const void* pData, int64 nData) : FElectraBitstreamReader(pData, nData)
+				{
+				}
+				int32 ReadMP4Length()
+				{
+					int32 Length = 0;
+					for (int32 i = 0; i < 4; ++i)
+					{
+						uint32 Bits = GetBits(8);
+						Length = (Length << 7) + (Bits & 0x7f);
+						if ((Bits & 0x80) == 0)
+							break;
+					}
+					return Length;
+				}
+			};
+			FMP4BitReader BitReader(InESDS.GetData(), InESDS.Num());
+
+			CSD.Empty();
+
+			if (BitReader.GetBits(8) != 3)
+			{
+				return false;
+			}
+
+			int32 ESSize = BitReader.ReadMP4Length();
+
+			ESID = (uint16)BitReader.GetBits(16);
+			bDependsOnStream = BitReader.GetBits(1) != 0;
+			bool bURLFlag = BitReader.GetBits(1) != 0;
+			bool bOCRflag = BitReader.GetBits(1) != 0;
+			StreamPriority = (uint8)BitReader.GetBits(5);
+			if (bDependsOnStream)
+			{
+				DependsOnStreamESID = BitReader.GetBits(16);
+			}
+			if (bURLFlag)
+			{
+				// Skip over the URL
+				uint32 urlLen = BitReader.GetBits(8);
+				BitReader.SkipBytes(urlLen);
+			}
+			if (bOCRflag)
+			{
+				// Skip the OCR ES ID
+				BitReader.SkipBits(16);
+			}
+
+			// Parse the config descriptor
+			if (BitReader.GetBits(8) != 4)
+			{
+				return false;
+			}
+			int32 ConfigDescrSize = BitReader.ReadMP4Length();
+			ObjectTypeID = static_cast<FObjectTypeID>(BitReader.GetBits(8));
+			StreamTypeID = static_cast<FStreamType>(BitReader.GetBits(6));
+			// Skip upstream flag
+			BitReader.SkipBits(1);
+			if (BitReader.GetBits(1) != 1)	// reserved, must be 1
+			{
+				return false;
+			}
+			BufferSize = BitReader.GetBits(24);
+			MaxBitrate = BitReader.GetBits(32);
+			AvgBitrate = BitReader.GetBits(32);
+			if (ConfigDescrSize > 13)
+			{
+				// Optional codec specific descriptor
+				if (BitReader.GetBits(8) != 5)
+				{
+					return false;
+				}
+				int32 CodecSize = BitReader.ReadMP4Length();
+				CSD.Reserve(CodecSize);
+				for (int32 i = 0; i < CodecSize; ++i)
+				{
+					CSD.Push(BitReader.GetBits(8));
+				}
+			}
+
+			// SL config (we do not need it, we require it to be there though as per the standard)
+			if (BitReader.GetBits(8) != 6)
+			{
+				return false;
+			}
+			int32 nSLSize = BitReader.ReadMP4Length();
+			if (nSLSize != 1)
+			{
+				return false;
+			}
+			if (BitReader.GetBits(8) != 2)
+			{
+				return false;
+			}
+
+			return true;
+		}
+
+
+
+
+
+
+
+
+
+
+
+
+
 		namespace AACUtils
 		{
 			int32 GetNumberOfChannelsFromChannelConfiguration(uint32 InChannelConfiguration)
@@ -105,6 +222,23 @@ namespace ElectraDecodersUtil
 			ExtAOT = 0;
 		}
 
+		FString FAACDecoderConfigurationRecord::GetFormatInfo() const
+		{
+			FString fi;
+			if (PSSignal > 0)
+			{
+				fi = TEXT("HE-AAC v2");
+			}
+			else if (SBRSignal > 0)
+			{
+				fi = TEXT("HE-AAC");
+			}
+			else
+			{
+				fi = TEXT("AAC");
+			}
+			return fi;
+		}
 
 		const TArray<uint8>& FAACDecoderConfigurationRecord::GetCodecSpecificData() const
 		{
diff --git a/Engine/Plugins/Media/ElectraCodecs/Source/ElectraDecoders/Private/Utils/MPEG/ElectraUtilsMPEGVideo.cpp b/Engine/Plugins/Media/ElectraCodecs/Source/ElectraDecoders/Private/Utils/MPEG/ElectraUtilsMPEGVideo.cpp
--- a/Engine/Plugins/Media/ElectraCodecs/Source/ElectraDecoders/Private/Utils/MPEG/ElectraUtilsMPEGVideo.cpp
+++ b/Engine/Plugins/Media/ElectraCodecs/Source/ElectraDecoders/Private/Utils/MPEG/ElectraUtilsMPEGVideo.cpp
@@ -1342,43 +1342,43 @@ namespace ElectraDecodersUtil
 		}
 
 
-		bool FSEImastering_display_colour_volume::ParseFromMessage(const FSEIMessage& InMessage)
+		bool ParseFromSEIMessage(FSEImastering_display_colour_volume& OutMDCV, const FSEIMessage& InMessage)
 		{
 			if (InMessage.PayloadType == FSEIMessage::EPayloadType::PT_mastering_display_colour_volume && InMessage.Message.Num() >= 24)
 			{
 				FElectraBitstreamReader BitReader(InMessage.Message.GetData(), InMessage.Message.Num());
 				for(int32 i=0; i<3; ++i)
 				{
-					display_primaries_x[i] = BitReader.GetBits(16);
-					display_primaries_y[i] = BitReader.GetBits(16);
+					OutMDCV.display_primaries_x[i] = BitReader.GetBits(16);
+					OutMDCV.display_primaries_y[i] = BitReader.GetBits(16);
 				}
-				white_point_x = BitReader.GetBits(16);
-				white_point_y = BitReader.GetBits(16);
-				max_display_mastering_luminance = BitReader.GetBits(32);
-				min_display_mastering_luminance = BitReader.GetBits(32);
+				OutMDCV.white_point_x = BitReader.GetBits(16);
+				OutMDCV.white_point_y = BitReader.GetBits(16);
+				OutMDCV.max_display_mastering_luminance = BitReader.GetBits(32);
+				OutMDCV.min_display_mastering_luminance = BitReader.GetBits(32);
 				return true;
 			}
 			return false;
 		}
 
-		bool FSEIcontent_light_level_info::ParseFromMessage(const FSEIMessage& InMessage)
+		bool ParseFromSEIMessage(FSEIcontent_light_level_info& OutCLLI, const FSEIMessage& InMessage)
 		{
 			if (InMessage.PayloadType == FSEIMessage::EPayloadType::PT_content_light_level_info && InMessage.Message.Num() >= 4)
 			{
 				FElectraBitstreamReader BitReader(InMessage.Message.GetData(), InMessage.Message.Num());
-				max_content_light_level = BitReader.GetBits(16);
-				max_pic_average_light_level = BitReader.GetBits(16);
+				OutCLLI.max_content_light_level = BitReader.GetBits(16);
+				OutCLLI.max_pic_average_light_level = BitReader.GetBits(16);
 				return true;
 			}
 			return false;
 		}
 
-		bool FSEIalternative_transfer_characteristics::ParseFromMessage(const FSEIMessage& InMessage)
+		bool ParseFromSEIMessage(FSEIalternative_transfer_characteristics& OutATC, const FSEIMessage& InMessage)
 		{
 			if (InMessage.PayloadType == FSEIMessage::EPayloadType::PT_alternative_transfer_characteristics && InMessage.Message.Num() >= 1)
 			{
 				FElectraBitstreamReader BitReader(InMessage.Message.GetData(), InMessage.Message.Num());
-				preferred_transfer_characteristics = BitReader.GetBits(8);
+				OutATC.preferred_transfer_characteristics = BitReader.GetBits(8);
 				return true;
 			}
 			return false;
diff --git a/Engine/Plugins/Media/ElectraCodecs/Source/ElectraDecoders/Private/Utils/MPEG/ElectraUtilsMPEGVideo_H264.cpp b/Engine/Plugins/Media/ElectraCodecs/Source/ElectraDecoders/Private/Utils/MPEG/ElectraUtilsMPEGVideo_H264.cpp
--- a/Engine/Plugins/Media/ElectraCodecs/Source/ElectraDecoders/Private/Utils/MPEG/ElectraUtilsMPEGVideo_H264.cpp
+++ b/Engine/Plugins/Media/ElectraCodecs/Source/ElectraDecoders/Private/Utils/MPEG/ElectraUtilsMPEGVideo_H264.cpp
@@ -315,6 +315,13 @@ namespace ElectraDecodersUtil
 				OutWidth = GetWidth() - cl - cr;
 				OutHeight = GetHeight() - ct - cb;
 			}
+			void FSequenceParameterSet::GetDisplaySize(uint32& OutWidth, uint32& OutHeight) const
+			{
+				int32 w,h;
+				GetDisplaySize(w, h);
+				OutWidth = (uint32) w;
+				OutHeight = (uint32) h;
+			}
 			void FSequenceParameterSet::GetCrop(int32& OutLeft, int32& OutRight, int32& OutTop, int32& OutBottom) const
 			{
 				if (frame_cropping_flag)
@@ -374,6 +381,130 @@ namespace ElectraDecodersUtil
 				return FFractionalValue();
 			}
 
+			FString FSequenceParameterSet::GetFormatInfo() const
+			{
+				return GetFormatInfo(profile_idc, level_idc, constraint_set0_flag | (constraint_set1_flag <<1) | (constraint_set2_flag << 2) | (constraint_set3_flag << 3) | (constraint_set4_flag << 4) | (constraint_set5_flag << 5));
+			}
+			FString FSequenceParameterSet::GetFormatInfo(uint32 InProfileIDC, uint32 InLevelIDC, uint32 InConstraintFlags)
+			{
+				FString fi;
+				if (InProfileIDC == 44)
+				{
+					fi = TEXT("CAVLC 4:4:4 Intra");
+				}
+				else if (InProfileIDC == 66 && (InConstraintFlags & 2) == 2)
+				{
+					fi = TEXT("Constrained Baseline");
+				}
+				else if (InProfileIDC == 66)
+				{
+					fi = TEXT("Baseline");
+				}
+				else if (InProfileIDC == 77)
+				{
+					fi = TEXT("Main");
+				}
+				else if (InProfileIDC == 83 && (InConstraintFlags & 32) == 32)
+				{
+					fi = TEXT("Scalable Constrained Baseline");
+				}
+				else if (InProfileIDC == 83)
+				{
+					fi = TEXT("Scalable Baseline");
+				}
+				else if (InProfileIDC == 86 && (InConstraintFlags & 32) == 32)
+				{
+					fi = TEXT("Scalable Constrained High");
+				}
+				else if (InProfileIDC == 86 && (InConstraintFlags & 8) == 8)
+				{
+					fi = TEXT("Scalable High Intra");
+				}
+				else if (InProfileIDC == 86)
+				{
+					fi = TEXT("Scalable High");
+				}
+				else if (InProfileIDC == 88)
+				{
+					fi = TEXT("Extended");
+				}
+				else if (InProfileIDC == 100 && (InConstraintFlags & 48) == 48)
+				{
+					fi = TEXT("Constrained High");
+				}
+				else if (InProfileIDC == 100 && (InConstraintFlags & 16) == 16)
+				{
+					fi = TEXT("Progressive High");
+				}
+				else if (InProfileIDC == 100)
+				{
+					fi = TEXT("High");
+				}
+				else if (InProfileIDC == 110 && (InConstraintFlags & 8) == 8)
+				{
+					fi = TEXT("High 10 Intra");
+				}
+				else if (InProfileIDC == 110)
+				{
+					fi = TEXT("High 10");
+				}
+				else if (InProfileIDC == 118)
+				{
+					fi = TEXT("Multiview High");
+				}
+				else if (InProfileIDC == 122 && (InConstraintFlags & 8) == 8)
+				{
+					fi = TEXT("High 4:2:2 Intra");
+				}
+				else if (InProfileIDC == 122)
+				{
+					fi = TEXT("High 4:2:2");
+				}
+				else if (InProfileIDC == 128)
+				{
+					fi = TEXT("Stereo High");
+				}
+				else if (InProfileIDC == 134)
+				{
+					fi = TEXT("MFC High");
+				}
+				else if (InProfileIDC == 135)
+				{
+					fi = TEXT("MFC Depth High");
+				}
+				else if (InProfileIDC == 138)
+				{
+					fi = TEXT("Multiview Depth High");
+				}
+				else if (InProfileIDC == 139)
+				{
+					fi = TEXT("Enhanced Multiview Depth High");
+				}
+				else if (InProfileIDC == 244 && (InConstraintFlags & 8) == 8)
+				{
+					fi = TEXT("High 4:4:4 Intra");
+				}
+				else if (InProfileIDC == 244)
+				{
+					fi = TEXT("High 4:4:4 Predictive");
+				}
+				else
+				{
+					fi = TEXT("Unkown");
+				}
+				fi.Append(TEXT(" Profile"));
+				if (InLevelIDC)
+				{
+					fi.Append(FString::Printf(TEXT(", level %d.%d"), InLevelIDC/10, InLevelIDC%10));
+				}
+				else
+				{
+					fi.Append(TEXT(" at unknown level"));
+				}
+				return fi;
+			}
+
+
 
 			bool ParseSequenceParameterSet(FSequenceParameterSet& OutSequenceParameterSet, const uint8* InBitstream, uint64 InBitstreamLenInBytes)
 			{
@@ -2094,6 +2225,274 @@ namespace ElectraDecodersUtil
 			}
 
 
+
+
+
+			bool FAVCDecoderConfigurationRecord::Parse(const TConstArrayView<uint8>& InDCR)
+			{
+				CSD.Empty();
+				RawData.Empty();
+
+				if (!InDCR.Num())
+				{
+					return false;
+				}
+
+				RawData = InDCR;
+				FElectraByteReader ByteReader(InDCR.GetData(), InDCR.Num());
+
+				if (!ByteReader.ReadByte(ConfigurationVersion) || ConfigurationVersion != 1)
+				{
+					return false;
+				}
+				if (!ByteReader.ReadByte(AVCProfileIndication))
+				{
+					return false;
+				}
+				if (!ByteReader.ReadByte(ProfileCompatibility))
+				{
+					return false;
+				}
+				if (!ByteReader.ReadByte(AVCLevelIndication))
+				{
+					return false;
+				}
+				if (!ByteReader.ReadByte(NALUnitLength))
+				{
+					return false;
+				}
+				NALUnitLength = (NALUnitLength & 3) + 1;
+				uint8 nSPS = 0;
+				if (!ByteReader.ReadByte(nSPS))
+				{
+					return false;
+				}
+				nSPS &= 31;
+				if (nSPS)
+				{
+					SequenceParameterSets.Reserve(nSPS);
+				}
+				int32 TotalSPSSize = 0;
+				for(int32 i=0; i<nSPS; ++i)
+				{
+					uint16 spsLen = 0;
+					if (!ByteReader.ReadByte(spsLen))
+					{
+						return false;
+					}
+					TArray<uint8>& sps = SequenceParameterSets.AddDefaulted_GetRef();
+					sps.Reserve(spsLen);
+					sps.SetNumUninitialized(spsLen);
+					if (!ByteReader.ReadBytes(sps.GetData(), spsLen))
+					{
+						return false;
+					}
+					TotalSPSSize += 4 + spsLen;		// 4 because we always use 32 bit startcode and not NALUnitLength
+				}
+				uint8 nPPS = 0;
+				if (!ByteReader.ReadByte(nPPS))
+				{
+					return false;
+				}
+				if (nPPS)
+				{
+					PictureParameterSets.Reserve(nPPS);
+				}
+				int32 TotalPPSSize = 0;
+				for(int32 i=0; i<nPPS; ++i)
+				{
+					uint16 ppsLen = 0;
+					if (!ByteReader.ReadByte(ppsLen))
+					{
+						return false;
+					}
+					TArray<uint8>& pps = PictureParameterSets.AddDefaulted_GetRef();
+					pps.Reserve(ppsLen);
+					pps.SetNumUninitialized(ppsLen);
+					if (!ByteReader.ReadBytes(pps.GetData(), ppsLen))
+					{
+						return false;
+					}
+					TotalPPSSize += 4 + ppsLen;		// 4 because we always use 32 bit startcode and not NALUnitLength
+				}
+
+				if (AVCProfileIndication != 66 && AVCProfileIndication != 77 && AVCProfileIndication != 88)
+				{
+					// At least according to the standard these values must appear.
+					// However, there are files that are of profile 100 but omit these values.
+					// Therefore let's do a quick check if we can read at least 4 more bytes.
+					if (ByteReader.BytesRemaining() >= 4)
+					{
+						if (!ByteReader.ReadByte(ChromaFormat))
+						{
+							return false;
+						}
+						ChromaFormat &= 3;
+						if (!ByteReader.ReadByte(BitDepthLumaMinus8))
+						{
+							return false;
+						}
+						BitDepthLumaMinus8 &= 7;
+						if (!ByteReader.ReadByte(BitDepthChromaMinus8))
+						{
+							return false;
+						}
+						BitDepthChromaMinus8 &= 7;
+						uint8 nSPSExt;
+						if (!ByteReader.ReadByte(nSPSExt))
+						{
+							return false;
+						}
+						if (nSPSExt)
+						{
+							SequenceParameterSetsExt.Reserve(nSPSExt);
+						}
+						for(int32 i=0; i<nSPSExt; ++i)
+						{
+							uint16 spsLenExt = 0;
+							if (!ByteReader.ReadByte(spsLenExt))
+							{
+								return false;
+							}
+							TArray<uint8>& spsExt = SequenceParameterSetsExt.AddDefaulted_GetRef();
+							spsExt.Reserve(spsLenExt);
+							spsExt.SetNumUninitialized(spsLenExt);
+							if (!ByteReader.ReadBytes(spsExt.GetData(), spsLenExt))
+							{
+								return false;
+							}
+						}
+						bHaveAdditionalProfileIndication = true;
+					}
+				}
+
+				int32 TotalCSDSize = TotalSPSSize + TotalPPSSize;
+				if (TotalCSDSize)
+				{
+					const TArray<uint8> StartCode { 0,0,0,1 };
+					CSD.Reserve(TotalCSDSize);
+					for(int32 i=0; i<SequenceParameterSets.Num(); ++i)
+					{
+						CSD.Append(StartCode);
+						CSD.Append(SequenceParameterSets[i]);
+					}
+					for(int32 i=0; i<PictureParameterSets.Num(); ++i)
+					{
+						CSD.Append(StartCode);
+						CSD.Append(PictureParameterSets[i]);
+					}
+				}
+				return true;
+			}
+
+			bool FAVCDecoderConfigurationRecord::CreateFromCodecSpecificData(const TConstArrayView<uint8>& InCSD)
+			{
+				CSD.Empty();
+				RawData.Empty();
+				if (InCSD.IsEmpty())
+				{
+					return false;
+				}
+
+				CSD = InCSD;
+				TArray<FNaluInfo> Nalus;
+				if (!ParseBitstreamForNALUs(Nalus, InCSD.GetData(), InCSD.Num()))
+				{
+					return false;
+				}
+
+				struct FArray
+				{
+					TArray<TArray<uint8>> NALUs;
+				};
+				FArray aSPS, aPPS;
+				for(int32 i=0; i<Nalus.Num(); ++i)
+				{
+					if (Nalus[i].Type == 7)
+					{
+						FSequenceParameterSet sps;
+						if (ParseSequenceParameterSet(sps, InCSD.GetData() + Nalus[i].Offset + Nalus[i].UnitLength, Nalus[i].Size))
+						{
+							AVCProfileIndication = sps.profile_idc;
+							ProfileCompatibility = (sps.constraint_set0_flag << 7) | (sps.constraint_set1_flag << 6) | (sps.constraint_set2_flag << 5) |
+													(sps.constraint_set3_flag << 4) | (sps.constraint_set4_flag << 3) | (sps.constraint_set5_flag << 2);
+							AVCLevelIndication = sps.level_idc;
+							NALUnitLength = 4;
+							ChromaFormat = sps.chroma_format_idc;
+							BitDepthLumaMinus8 = sps.bit_depth_luma_minus8;
+							BitDepthChromaMinus8 = sps.bit_depth_chroma_minus8;
+							bHaveAdditionalProfileIndication = false;
+
+							TArray<uint8>& NewSPS(SequenceParameterSets.Emplace_GetRef());
+							NewSPS.Append(InCSD.GetData() + Nalus[i].Offset + Nalus[i].UnitLength, Nalus[i].Size);
+						}
+						else
+						{
+							return false;
+						}
+						aSPS.NALUs.Emplace_GetRef() = MakeArrayView<const uint8>(InCSD.GetData() + Nalus[i].Offset + Nalus[i].UnitLength, Nalus[i].Size);
+					}
+					else if (Nalus[i].Type == 8)
+					{
+						aPPS.NALUs.Emplace_GetRef() = MakeArrayView<const uint8>(InCSD.GetData() + Nalus[i].Offset + Nalus[i].UnitLength, Nalus[i].Size);
+					}
+				}
+				if (aSPS.NALUs.IsEmpty() || aPPS.NALUs.IsEmpty())
+				{
+					return false;
+				}
+
+				// Create the raw configuration record now.
+				ConfigurationVersion = 1;
+				FElectraBitstreamWriter wr;
+				wr.PutBits(ConfigurationVersion, 8);
+				wr.PutBits(AVCProfileIndication, 8);
+				wr.PutBits(ProfileCompatibility, 8);
+				wr.PutBits(AVCLevelIndication, 8);
+				wr.PutBits(63U, 6);
+				wr.PutBits(NALUnitLength-1, 2);
+				wr.PutBits(7U, 3);
+				wr.PutBits(static_cast<uint32>(aSPS.NALUs.Num()), 5);
+				for(int32 j=0; j<aSPS.NALUs.Num(); ++j)
+				{
+					const TArray<uint8>& d(aSPS.NALUs[j]);
+					wr.PutBits(static_cast<uint32>(d.Num()), 16);
+					for(int32 k=0; k<d.Num(); ++k)
+					{
+						wr.PutBits(d[k], 8);
+					}
+				}
+				wr.PutBits(static_cast<uint32>(aPPS.NALUs.Num()), 8);
+				for(int32 j=0; j<aPPS.NALUs.Num(); ++j)
+				{
+					const TArray<uint8>& d(aPPS.NALUs[j]);
+					wr.PutBits(static_cast<uint32>(d.Num()), 16);
+					for(int32 k=0; k<d.Num(); ++k)
+					{
+						wr.PutBits(d[k], 8);
+					}
+				}
+				if (AVCProfileIndication != 66 && AVCProfileIndication != 77 && AVCProfileIndication != 88)
+				{
+					wr.PutBits(63U, 6);
+					wr.PutBits(ChromaFormat, 2);
+					wr.PutBits(31U, 5);
+					wr.PutBits(BitDepthLumaMinus8, 3);
+					wr.PutBits(31U, 5);
+					wr.PutBits(BitDepthChromaMinus8, 3);
+					wr.PutBits(0U, 8);
+				}
+				wr.GetArray(RawData);
+				return true;
+			}
+
+			FString FAVCDecoderConfigurationRecord::GetFormatInfo() const
+			{
+				FString fi(TEXT("AVC, "));
+				fi.Append(FSequenceParameterSet::GetFormatInfo(AVCProfileIndication, AVCLevelIndication, 0));
+				return fi;
+			}
+
 		} // namespace H264
 	} // namespace MPEG
 } // namespace ElectraDecodersUtil
diff --git a/Engine/Plugins/Media/ElectraCodecs/Source/ElectraDecoders/Private/Utils/MPEG/ElectraUtilsMPEGVideo_H265.cpp b/Engine/Plugins/Media/ElectraCodecs/Source/ElectraDecoders/Private/Utils/MPEG/ElectraUtilsMPEGVideo_H265.cpp
--- a/Engine/Plugins/Media/ElectraCodecs/Source/ElectraDecoders/Private/Utils/MPEG/ElectraUtilsMPEGVideo_H265.cpp
+++ b/Engine/Plugins/Media/ElectraCodecs/Source/ElectraDecoders/Private/Utils/MPEG/ElectraUtilsMPEGVideo_H265.cpp
@@ -218,6 +218,13 @@ namespace ElectraDecodersUtil
 				OutWidth = GetWidth() - cl - cr;
 				OutHeight = GetHeight() - ct - cb;
 			}
+			void FSequenceParameterSet::GetDisplaySize(uint32& OutWidth, uint32& OutHeight) const
+			{
+				int32 w,h;
+				GetDisplaySize(w, h);
+				OutWidth = (uint32) w;
+				OutHeight = (uint32) h;
+			}
 			void FSequenceParameterSet::GetCrop(int32& OutLeft, int32& OutRight, int32& OutTop, int32& OutBottom) const
 			{
 				OutLeft = conf_win_left_offset * SubWidthC;
@@ -299,6 +306,39 @@ namespace ElectraDecodersUtil
 			}
 
 
+			FString FSequenceParameterSet::GetFormatInfo() const
+			{
+				return GetFormatInfo(profile_tier_level.general_profile_idc, profile_tier_level.general_level_idc, GetConstraintFlags());
+			}
+
+			FString FSequenceParameterSet::GetFormatInfo(uint32 InProfileIDC, uint32 InLevelIDC, uint64 InConstraintFlags)
+			{
+				FString fi;
+				if (InProfileIDC == 1)
+				{
+					fi = TEXT("Main");
+				}
+				else if (InProfileIDC == 2)
+				{
+					fi = TEXT("Main 10");
+				}
+				else
+				{
+					fi = TEXT("Unknown");
+				}
+				fi.Append(TEXT(" Profile"));
+				if (InLevelIDC)
+				{
+					fi.Append(FString::Printf(TEXT(", level %d.%d"), InLevelIDC/30, InLevelIDC%30));
+				}
+				else
+				{
+					fi.Append(TEXT(" at unknown level"));
+				}
+				return fi;
+			}
+
+
 			bool FHRDParameters::Parse(FBitstreamReader& br, bool commonInfPresentFlag, uint32 maxNumSubLayersMinus1)
 			{
 				#define RANGE_CHECK_FAILURE(expr, elem)															\
@@ -2519,6 +2559,293 @@ namespace ElectraDecodersUtil
 			}
 
 
+
+
+
+
+
+
+
+
+			bool FHEVCDecoderConfigurationRecord::Parse(const TConstArrayView<uint8>& InDCR)
+			{
+				CSD.Empty();
+				RawData.Empty();
+
+				if (InDCR.Num() >= 23)
+				{
+					RawData = InDCR;
+
+					FElectraBitstreamReader BitReader(InDCR.GetData(), InDCR.Num());
+
+					ConfigurationVersion = (uint8)BitReader.GetBits(8);
+					if (ConfigurationVersion != 1)
+					{
+						return false;
+					}
+					GeneralProfileSpace = (uint8)BitReader.GetBits(2);
+					GeneralTierFlag = (uint8)BitReader.GetBits(1);
+					GeneralProfileIDC = (uint8)BitReader.GetBits(5);
+					GeneralProfileCompatibilityFlags  = BitReader.GetBits(32);
+					for(int32 i=0; i<48; ++i)
+					{
+						GeneralConstraintIndicatorFlags = (GeneralConstraintIndicatorFlags << 1U) | uint64(BitReader.GetBits(1));
+					}
+					GeneralLevelIDC = (uint8)BitReader.GetBits(8);
+					if (BitReader.GetBits(4) != 15)
+					{
+						return false;
+					}
+					MinSpatialSegmentationIDC = (uint16)BitReader.GetBits(12);
+					if (BitReader.GetBits(6) != 63)
+					{
+						return false;
+					}
+					ParallelismType = (uint8)BitReader.GetBits(2);
+					if (BitReader.GetBits(6) != 63)
+					{
+						return false;
+					}
+					ChromaFormat = (uint8)BitReader.GetBits(2);
+					if (BitReader.GetBits(5) != 31)
+					{
+						return false;
+					}
+					BitDepthLumaMinus8 = (uint8)BitReader.GetBits(3);
+					if (BitReader.GetBits(5) != 31)
+					{
+						return false;
+					}
+					BitDepthChromaMinus8 = (uint8)BitReader.GetBits(3);
+					AverageFrameRate = (uint16)BitReader.GetBits(16);
+					ConstantFrameRate = (uint8)BitReader.GetBits(2);
+					NumTemporalLayers = (uint8)BitReader.GetBits(3);
+					TemporalIdNested = (uint8)BitReader.GetBits(1);
+					NALUnitLengthMinus1 = (uint8)BitReader.GetBits(2);
+					uint32 numArrays = BitReader.GetBits(8);
+					Arrays.Reserve(numArrays);
+					for(uint32 i=0; i<numArrays; ++i)
+					{
+						FArray &a = Arrays.AddDefaulted_GetRef();
+						a.Completeness = (uint8)BitReader.GetBits(1);
+						if (BitReader.GetBits(1) != 0)
+						{
+							return false;
+						}
+						a.NALUnitType = (uint8)BitReader.GetBits(6);
+						uint32 numNALUs = BitReader.GetBits(16);
+						a.NALUs.Reserve(numNALUs);
+						for(uint32 j=0; j<numNALUs; ++j)
+						{
+							uint16 naluLen = (uint16)BitReader.GetBits(16);
+							TArray<uint8> &n = a.NALUs.AddDefaulted_GetRef();
+							n.Reserve(naluLen);
+							for(uint32 k=0; k<naluLen; ++k)
+							{
+								n.Emplace((uint8)BitReader.GetBits(8));
+							}
+						}
+					}
+					if (BitReader.GetRemainingBits() != 0)
+					{
+						return false;
+					}
+
+					// CSD
+					for(int32 i=0; i<Arrays.Num(); ++i)
+					{
+						const FArray& a = Arrays[i];
+						for(int32 j=0; j<a.NALUs.Num(); ++j)
+						{
+							CSD.Add(0);
+							CSD.Add(0);
+							CSD.Add(0);
+							CSD.Add(1);
+							CSD.Append(a.NALUs[j]);
+						}
+
+						// SPS nut?
+						if (a.NALUnitType == 33)
+						{
+							TArray<uint8>& sps(SequenceParameterSets.Emplace_GetRef());
+							for(int32 k=0; k<a.NALUs.Num(); ++k)
+							{
+								sps.Append(a.NALUs[k]);
+								if (!ParseSequenceParameterSet(ParsedSequenceParameterSets.Emplace_GetRef(), a.NALUs[k].GetData(), a.NALUs[k].Num()))
+								{
+									ParsedSequenceParameterSets.RemoveAt(ParsedSequenceParameterSets.Num()-1);
+								}
+							}
+						}
+					}
+					return true;
+				}
+				return false;
+			}
+
+			FString FHEVCDecoderConfigurationRecord::GetCodecSpecifierRFC6381(const TCHAR* InPrefix) const
+			{
+				/*
+				// If we have a parsed SPS we use that.
+				if (ParsedSequenceParameterSets.Num())
+				{
+					return ParsedSequenceParameterSets[0].GetRFC6381(InPrefix);
+				}
+				*/
+				// Construct one from the general information as per ISO/IEC 14496-15:2022 Annex E.3
+				FString cfs;
+				uint64 ConstraintFlags = GeneralConstraintIndicatorFlags;
+				bool bNonZero = false;
+				for(int32 i=5; i>=0; --i, ConstraintFlags >>= 8)
+				{
+					if ((ConstraintFlags & 255) == 0 && !bNonZero)
+					{
+						continue;
+					}
+					bNonZero = true;
+					cfs = FString::Printf(TEXT(".%02X"), (uint8)(ConstraintFlags & 255)) + cfs;
+				}
+				const uint32 cf = BitReverse32(GeneralProfileCompatibilityFlags);
+				FString rfc;
+				if (GeneralProfileSpace == 0)
+				{
+					rfc = FString::Printf(TEXT("%s.%d.%X.%c%d%s"), InPrefix, GeneralProfileIDC, cf, GeneralTierFlag ? TCHAR('H') : TCHAR('L'), GeneralLevelIDC, *cfs);
+				}
+				else
+				{
+					rfc = FString::Printf(TEXT("%s.%c%d.%X.%c%d%s"), InPrefix, TCHAR(TCHAR('A')+GeneralProfileSpace-1), GeneralProfileIDC, cf, GeneralTierFlag ? TCHAR('H') : TCHAR('L'), GeneralLevelIDC, *cfs);
+				}
+				return rfc;
+			}
+
+			FString FHEVCDecoderConfigurationRecord::GetFormatInfo() const
+			{
+				FString fi(TEXT("HEVC, "));
+				fi.Append(FSequenceParameterSet::GetFormatInfo(GeneralProfileIDC, GeneralLevelIDC, GeneralConstraintIndicatorFlags));
+				return fi;
+			}
+
+			bool FHEVCDecoderConfigurationRecord::CreateFromCodecSpecificData(const TConstArrayView<uint8>& InCSD)
+			{
+				CSD.Empty();
+				RawData.Empty();
+				if (InCSD.IsEmpty())
+				{
+					return false;
+				}
+
+				CSD = InCSD;
+				TArray<FNaluInfo> Nalus;
+				if (!ParseBitstreamForNALUs(Nalus, InCSD.GetData(), InCSD.Num()))
+				{
+					return false;
+				}
+				FArray aVPS, aSPS, aPPS;
+				for(int32 i=0; i<Nalus.Num(); ++i)
+				{
+					if (Nalus[i].Type == 32)
+					{
+						aVPS.NALUnitType = 32;
+						aVPS.NALUs.Emplace_GetRef() = MakeArrayView<const uint8>(InCSD.GetData() + Nalus[i].Offset + Nalus[i].UnitLength, Nalus[i].Size);
+					}
+					else if (Nalus[i].Type == 33)
+					{
+						FSequenceParameterSet sps;
+						if (ParseSequenceParameterSet(sps, InCSD.GetData() + Nalus[i].Offset + Nalus[i].UnitLength, Nalus[i].Size))
+						{
+							GeneralProfileSpace = sps.profile_tier_level.general_profile_space;
+							GeneralTierFlag = sps.profile_tier_level.general_tier_flag;
+							GeneralProfileIDC = sps.profile_tier_level.general_profile_idc;
+							GeneralProfileCompatibilityFlags = sps.profile_tier_level.general_profile_compatibility_flags;
+							GeneralConstraintIndicatorFlags = sps.profile_tier_level.general_constraint_indicator_flags;
+							GeneralLevelIDC = sps.profile_tier_level.general_level_idc;
+							check(sps.vui_parameters_present_flag);
+							MinSpatialSegmentationIDC = sps.vui_parameters.min_spatial_segmentation_idc;
+							ParallelismType = 0;	// for simplicity as 0 indicates `unknown`.
+							ChromaFormat = sps.chroma_format_idc;
+							BitDepthLumaMinus8 = sps.bit_depth_luma_minus8;
+							BitDepthChromaMinus8 = sps.bit_depth_chroma_minus8;
+							AverageFrameRate = 0;	// unknown
+							ConstantFrameRate = 0;	// may or may not
+							NumTemporalLayers = 0;	// unknown
+							TemporalIdNested = 0;	// may or may not
+							NALUnitLengthMinus1 = 3;
+
+							TArray<uint8>& NewSPS(SequenceParameterSets.Emplace_GetRef());
+							NewSPS.Append(InCSD.GetData() + Nalus[i].Offset + Nalus[i].UnitLength, Nalus[i].Size);
+							ParsedSequenceParameterSets.Emplace(MoveTemp(sps));
+						}
+						else
+						{
+							return false;
+						}
+						aSPS.NALUnitType = 33;
+						aSPS.NALUs.Emplace_GetRef() = MakeArrayView<const uint8>(InCSD.GetData() + Nalus[i].Offset + Nalus[i].UnitLength, Nalus[i].Size);
+					}
+					else if (Nalus[i].Type == 34)
+					{
+						aPPS.NALUnitType = 34;
+						aPPS.NALUs.Emplace_GetRef() = MakeArrayView<const uint8>(InCSD.GetData() + Nalus[i].Offset + Nalus[i].UnitLength, Nalus[i].Size);
+					}
+				}
+				if (aVPS.NALUs.IsEmpty() || aSPS.NALUs.IsEmpty() || aPPS.NALUs.IsEmpty())
+				{
+					return false;
+				}
+				Arrays.Emplace(MoveTemp(aVPS));
+				Arrays.Emplace(MoveTemp(aSPS));
+				Arrays.Emplace(MoveTemp(aPPS));
+
+				// Create the raw configuration record now.
+				ConfigurationVersion = 1;
+				FElectraBitstreamWriter wr;
+				wr.PutBits(ConfigurationVersion, 8);
+				wr.PutBits(GeneralProfileSpace, 2);
+				wr.PutBits(GeneralTierFlag, 1);
+				wr.PutBits(GeneralProfileIDC, 5);
+				wr.PutBits(GeneralProfileCompatibilityFlags, 32);
+				wr.PutBits64(GeneralConstraintIndicatorFlags, 48);
+				wr.PutBits(GeneralLevelIDC, 8);
+				wr.PutBits(15U, 4);
+				wr.PutBits(MinSpatialSegmentationIDC, 12);
+				wr.PutBits(63U, 6);
+				wr.PutBits(ParallelismType, 2);
+				wr.PutBits(63U, 6);
+				wr.PutBits(ChromaFormat, 2);
+				wr.PutBits(31U, 5);
+				wr.PutBits(BitDepthLumaMinus8, 3);
+				wr.PutBits(31U, 5);
+				wr.PutBits(BitDepthChromaMinus8, 3);
+				wr.PutBits(AverageFrameRate, 16);
+				wr.PutBits(ConstantFrameRate, 2);
+				wr.PutBits(NumTemporalLayers, 3);
+				wr.PutBits(TemporalIdNested, 1);
+				wr.PutBits(NALUnitLengthMinus1, 2);
+				wr.PutBits(static_cast<uint32>(Arrays.Num()), 8);
+				for(int32 i=0; i<Arrays.Num(); ++i)
+				{
+					const FArray& a(Arrays[i]);
+					wr.PutBits(a.Completeness, 1);
+					wr.PutBits(0U, 1);
+					wr.PutBits(a.NALUnitType, 6);
+					wr.PutBits(static_cast<uint32>(a.NALUs.Num()), 16);
+					for(int32 j=0; j<a.NALUs.Num(); ++j)
+					{
+						const TArray<uint8>& d(a.NALUs[j]);
+						wr.PutBits(static_cast<uint32>(d.Num()), 16);
+						for(int32 k=0; k<d.Num(); ++k)
+						{
+							wr.PutBits(d[k], 8);
+						}
+					}
+				}
+				wr.GetArray(RawData);
+				return true;
+			}
+
+
+
+
 		} // namespace H265
 	} // namespace MPEG
 } // namespace ElectraDecodersUtil
diff --git a/Engine/Plugins/Media/ElectraCodecs/Source/ElectraDecoders/Public/Utils/MPEG/ElectraUtilsMPEGAudio.h b/Engine/Plugins/Media/ElectraCodecs/Source/ElectraDecoders/Public/Utils/MPEG/ElectraUtilsMPEGAudio.h
--- a/Engine/Plugins/Media/ElectraCodecs/Source/ElectraDecoders/Public/Utils/MPEG/ElectraUtilsMPEGAudio.h
+++ b/Engine/Plugins/Media/ElectraCodecs/Source/ElectraDecoders/Public/Utils/MPEG/ElectraUtilsMPEGAudio.h
@@ -10,6 +10,63 @@ namespace ElectraDecodersUtil
 	namespace MPEG
 	{
 
+		class FESDescriptor
+		{
+		public:
+			ELECTRADECODERS_API bool Parse(const TConstArrayView<uint8>& InESDS);
+
+			inline const TArray<uint8>& GetCodecSpecificData() const
+			{ return CSD; }
+
+			inline uint32 GetBufferSize() const
+			{ return BufferSize; }
+
+			inline uint32 GetMaxBitrate() const
+			{ return MaxBitrate; }
+
+			inline uint32 GetAvgBitrate() const
+			{ return AvgBitrate; }
+
+			// See http://mp4ra.org/#/object_types
+			enum class FObjectTypeID
+			{
+				Unknown = 0,
+				Text_Stream = 8,
+				MPEG4_Video = 0x20,
+				H264 = 0x21,
+				H264_ParameterSets = 0x22,
+				H265 = 0x23,
+				MPEG4_Audio = 0x40,
+				MPEG1_Audio = 0x6b
+			};
+			enum class FStreamType
+			{
+				Unknown = 0,
+				VisualStream = 4,
+				AudioStream = 5
+			};
+
+			inline FObjectTypeID GetObjectTypeID() const
+			{ return ObjectTypeID; }
+
+			inline FStreamType GetStreamType() const
+			{ return StreamTypeID; }
+
+		private:
+			TArray<uint8> RawData;
+			TArray<uint8> CSD;
+			FObjectTypeID ObjectTypeID = FObjectTypeID::Unknown;
+			FStreamType StreamTypeID = FStreamType::Unknown;
+			uint32 BufferSize = 0;
+			uint32 MaxBitrate = 0;
+			uint32 AvgBitrate = 0;
+			uint16 ESID = 0;
+			uint16 DependsOnStreamESID = 0;
+			uint8 StreamPriority = 0;
+			bool bDependsOnStream = false;
+		};
+
+
 		class ELECTRADECODERS_API FAACDecoderConfigurationRecord
 		{
 		public:
@@ -19,27 +76,32 @@ namespace ElectraDecodersUtil
 			void Reset();
 			const TArray<uint8>& GetCodecSpecificData() const;
 
-			void SetRawData(const TArray<uint8>& InRawData)
-			{
-				RawData = InRawData;
-			}
-			const TArray<uint8>& GetRawData() const
-			{
-				return RawData;
-			}
-
-			int32		SBRSignal;
-			int32		PSSignal;
-			uint32		ChannelConfiguration;
-			uint32		SamplingFrequencyIndex;
-			uint32		SamplingRate;
-			uint32		ExtSamplingFrequencyIndex;
-			uint32		ExtSamplingFrequency;
-			uint32		AOT;
-			uint32		ExtAOT;
+			inline bool Parse(const TConstArrayView<uint8>& InDCR)
+			{ return ParseFrom(InDCR.GetData(), InDCR.Num()); }
+
+			inline void SetRawData(const TArray<uint8>& InRawData)
+			{ RawData = InRawData; }
+
+			inline const TArray<uint8>& GetRawData() const
+			{ return RawData; }
+
+			inline FString GetCodecSpecifierRFC6381() const
+			{ return FString::Printf(TEXT("mp4a.40.%u"), ExtAOT ? ExtAOT : AOT); }
+
+			FString GetFormatInfo() const;
+
+			int32 SBRSignal;
+			int32 PSSignal;
+			uint32 ChannelConfiguration;
+			uint32 SamplingFrequencyIndex;
+			uint32 SamplingRate;
+			uint32 ExtSamplingFrequencyIndex;
+			uint32 ExtSamplingFrequency;
+			uint32 AOT;
+			uint32 ExtAOT;
 		private:
-			TArray<uint8>	CodecSpecificData;
-			TArray<uint8>	RawData;
+			TArray<uint8> CodecSpecificData;
+			TArray<uint8> RawData;
 		};
 
 
diff --git a/Engine/Plugins/Media/ElectraCodecs/Source/ElectraDecoders/Public/Utils/MPEG/ElectraUtilsMPEGVideo.h b/Engine/Plugins/Media/ElectraCodecs/Source/ElectraDecoders/Public/Utils/MPEG/ElectraUtilsMPEGVideo.h
--- a/Engine/Plugins/Media/ElectraCodecs/Source/ElectraDecoders/Public/Utils/MPEG/ElectraUtilsMPEGVideo.h
+++ b/Engine/Plugins/Media/ElectraCodecs/Source/ElectraDecoders/Public/Utils/MPEG/ElectraUtilsMPEGVideo.h
@@ -456,6 +456,7 @@ namespace ElectraDecodersUtil
 		{
 			enum EPayloadType
 			{
+				PT_pic_timing = 1,
 				PT_user_data_registered_itu_t_t35 = 4,
 				// PT_tone_mapping_info = 23,							// not expected to be used (application custom)
 				PT_mastering_display_colour_volume = 137,
@@ -486,21 +487,21 @@ namespace ElectraDecodersUtil
 			uint16 white_point_y = 0;
 			uint32 max_display_mastering_luminance = 0;
 			uint32 min_display_mastering_luminance = 0;
-			bool ELECTRADECODERS_API ParseFromMessage(const FSEIMessage& InMessage);
 		};
+		bool ELECTRADECODERS_API ParseFromSEIMessage(FSEImastering_display_colour_volume& OutMDCV, const FSEIMessage& InMessage);
 
 		struct FSEIcontent_light_level_info
 		{
 			uint16 max_content_light_level = 0;
 			uint16 max_pic_average_light_level = 0;
-			bool ELECTRADECODERS_API ParseFromMessage(const FSEIMessage& InMessage);
 		};
+		bool ELECTRADECODERS_API ParseFromSEIMessage(FSEIcontent_light_level_info& OutCLLI, const FSEIMessage& InMessage);
 
 		struct FSEIalternative_transfer_characteristics
 		{
 			uint8 preferred_transfer_characteristics = 0;
-			bool ELECTRADECODERS_API ParseFromMessage(const FSEIMessage& InMessage);
 		};
+		bool ELECTRADECODERS_API ParseFromSEIMessage(FSEIalternative_transfer_characteristics& OutATC, const FSEIMessage& InMessage);
 
 		struct FSEIambient_viewing_environment
 		{
@@ -524,6 +525,41 @@ namespace ElectraDecodersUtil
 			uint32 ccv_avg_luminance_value = 0;
 		};
 
+
+		struct FCommonColorimetry
+		{
+			uint8 colour_primaries = 2;
+			uint8 transfer_characteristics = 2;
+			uint8 matrix_coeffs = 2;
+			uint8 video_full_range_flag = 0;
+			uint8 video_format = 5;
+		};
+
+		struct FCommonPictureTiming
+		{
+			// Calculated value as per:
+			//   clockTimestamp = ( ( hH * 60 + mM ) * 60 + sS ) * time_scale + nFrames * ( num_units_in_tick * ( 1 + nuit_field_based_flag ) ) + tOffset,
+			// can only be valid when there is timing information.
+			int64 clockTimestamp = 0;
+			// Values from a pic_timing() SEI in H.264 or from a time_code() SEI in H.265
+			uint32 num_units_in_tick = 0;			// from the SPS
+			uint32 time_scale = 0;					// from the SPS
+			int32 time_offset = 0;
+			uint16 n_frames = 0;
+			uint8 timing_info_present_flag = 0;		// from the SPS
+			uint8 clock_timestamp_flag = 0;
+			uint8 ct_type = 0;
+			uint8 nuit_field_based_flag = 0;
+			uint8 counting_type = 0;
+			uint8 full_timestamp_flag = 0;
+			uint8 discontinuity_flag = 0;
+			uint8 cnt_dropped_flag = 0;
+			uint8 seconds_value = 0;
+			uint8 minutes_value = 0;
+			uint8 hours_value = 0;
+			uint8 FromH26x = 0;		// last digit of the codec this comes from. 4=H.264, 5=H.265, etc. In case values need different interpretation.
+		};
+
 	} // namespace MPEG
 
 } // namespace ElectraDecodersUtil
diff --git a/Engine/Plugins/Media/ElectraCodecs/Source/ElectraDecoders/Public/Utils/MPEG/ElectraUtilsMPEGVideo_H264.h b/Engine/Plugins/Media/ElectraCodecs/Source/ElectraDecoders/Public/Utils/MPEG/ElectraUtilsMPEGVideo_H264.h
--- a/Engine/Plugins/Media/ElectraCodecs/Source/ElectraDecoders/Public/Utils/MPEG/ElectraUtilsMPEGVideo_H264.h
+++ b/Engine/Plugins/Media/ElectraCodecs/Source/ElectraDecoders/Public/Utils/MPEG/ElectraUtilsMPEGVideo_H264.h
@@ -92,7 +92,7 @@ namespace ElectraDecodersUtil
 				}
 			};
 
-			struct ELECTRADECODERS_API FSequenceParameterSet
+			struct FSequenceParameterSet
 			{
 				FSequenceParameterSet()
 				{
@@ -177,14 +177,17 @@ namespace ElectraDecodersUtil
 
 					ExpectedDeltaPerPicOrderCntCycle = 0;
 				}
-				int32 GetMaxDPBSize() const;
-				int32 GetDPBSize() const;
-				int32 GetWidth() const;
-				int32 GetHeight() const;
-				void GetDisplaySize(int32& OutWidth, int32& OutHeight) const;
-				void GetCrop(int32& OutLeft, int32& OutRight, int32& OutTop, int32& OutBottom) const;
-				void GetAspect(int32& OutSarW, int32& OutSarH) const;
-				FFractionalValue GetTiming() const;
+				ELECTRADECODERS_API int32 GetMaxDPBSize() const;
+				ELECTRADECODERS_API int32 GetDPBSize() const;
+				ELECTRADECODERS_API int32 GetWidth() const;
+				ELECTRADECODERS_API int32 GetHeight() const;
+				ELECTRADECODERS_API void GetDisplaySize(int32& OutWidth, int32& OutHeight) const;
+				ELECTRADECODERS_API void GetDisplaySize(uint32& OutWidth, uint32& OutHeight) const;
+				ELECTRADECODERS_API void GetCrop(int32& OutLeft, int32& OutRight, int32& OutTop, int32& OutBottom) const;
+				ELECTRADECODERS_API void GetAspect(int32& OutSarW, int32& OutSarH) const;
+				ELECTRADECODERS_API FFractionalValue GetTiming() const;
+				ELECTRADECODERS_API FString GetFormatInfo() const;
+				ELECTRADECODERS_API static FString GetFormatInfo(uint32 InProfileIDC, uint32 InLevelIDC, uint32 InConstraintFlags);
 
 				uint32 profile_idc;											// u(8)
 				uint8 constraint_set0_flag;									// u(1)
@@ -576,11 +579,68 @@ namespace ElectraDecodersUtil
 			};
 
 
-			bool ELECTRADECODERS_API ParseSequenceParameterSet(FSequenceParameterSet& OutSequenceParameterSet, const uint8* InBitstream, uint64 InBitstreamLenInBytes);
 			bool ELECTRADECODERS_API ParseSequenceParameterSet(TMap<uint32, FSequenceParameterSet>& InOutSequenceParameterSets, const uint8* InBitstream, uint64 InBitstreamLenInBytes);
 			bool ELECTRADECODERS_API ParsePictureParameterSet(TMap<uint32, FPictureParameterSet>& InOutPictureParameterSets, const TMap<uint32, FSequenceParameterSet>& InSequenceParameterSets, const uint8* InBitstream, uint64 InBitstreamLenInBytes);
 			bool ELECTRADECODERS_API ParseSliceHeader(TUniquePtr<FRBSP>& OutRBSP, FBitstreamReader& OutRBSPReader, FSliceHeader& OutSlice, const TMap<uint32, FSequenceParameterSet>& InSequenceParameterSets, const TMap<uint32, FPictureParameterSet>& InPictureParameterSets, const uint8* InBitstream, uint64 InBitstreamLenInBytes);
 
+			bool ELECTRADECODERS_API ParseSequenceParameterSet(FSequenceParameterSet& OutSequenceParameterSet, const uint8* InBitstream, uint64 InBitstreamLenInBytes);
+			inline bool ParseSequenceParameterSet(FSequenceParameterSet& OutSequenceParameterSet, const TArray<uint8>& InData)
+			{ return ParseSequenceParameterSet(OutSequenceParameterSet, InData.GetData(), InData.Num()); }
+
+
+
+
+
+			/**
+			 * ISO/IEC 14496-15:2022 - 5.3.2 Decoder configuration information
+			 */
+			class FAVCDecoderConfigurationRecord
+			{
+			public:
+				ELECTRADECODERS_API bool Parse(const TConstArrayView<uint8>& InDCR);
+				ELECTRADECODERS_API bool CreateFromCodecSpecificData(const TConstArrayView<uint8>& InCSD);
+
+				inline const TArray<uint8>& GetRawData() const
+				{ return RawData; }
+				inline const TArray<uint8>& GetCodecSpecificData() const
+				{ return CSD; }
+				inline const TArray<TArray<uint8>>& GetSequenceParameterSets() const
+				{ return SequenceParameterSets; }
+				inline const TArray<TArray<uint8>>& GetPictureParameterSets() const
+				{ return PictureParameterSets; }
+
+				inline uint8 GetAVCProfileIndication() const
+				{ return AVCProfileIndication; }
+				inline uint8 GetProfileCompatibility() const
+				{ return ProfileCompatibility; }
+				inline uint8 GetAVCLevelIndication() const
+				{ return AVCLevelIndication; }
+
+				inline FString GetCodecSpecifierRFC6381(const TCHAR* InPrefix) const
+				{ return FString::Printf(TEXT("%s.%02x%02x%02x"), InPrefix, GetAVCProfileIndication(), GetProfileCompatibility(), GetAVCLevelIndication()); }
+
+				ELECTRADECODERS_API FString GetFormatInfo() const;
+
+			private:
+				TArray<TArray<uint8>> SequenceParameterSets;
+				TArray<TArray<uint8>> PictureParameterSets;
+				TArray<TArray<uint8>> SequenceParameterSetsExt;
+				TArray<uint8> CSD;
+				TArray<uint8> RawData;
+				uint8 ConfigurationVersion = 0;
+				uint8 AVCProfileIndication = 0;
+				uint8 ProfileCompatibility = 0;
+				uint8 AVCLevelIndication = 0;
+				uint8 NALUnitLength = 0;
+				uint8 ChromaFormat = 0;
+				uint8 BitDepthLumaMinus8 = 0;
+				uint8 BitDepthChromaMinus8 = 0;
+				bool bHaveAdditionalProfileIndication = false;
+			};
+
+
+
+
 		} // namespace H264
 	} // namespace MPEG
 } // namespace ElectraDecodersUtil
diff --git a/Engine/Plugins/Media/ElectraCodecs/Source/ElectraDecoders/Public/Utils/MPEG/ElectraUtilsMPEGVideo_H265.h b/Engine/Plugins/Media/ElectraCodecs/Source/ElectraDecoders/Public/Utils/MPEG/ElectraUtilsMPEGVideo_H265.h
--- a/Engine/Plugins/Media/ElectraCodecs/Source/ElectraDecoders/Public/Utils/MPEG/ElectraUtilsMPEGVideo_H265.h
+++ b/Engine/Plugins/Media/ElectraCodecs/Source/ElectraDecoders/Public/Utils/MPEG/ElectraUtilsMPEGVideo_H265.h
@@ -622,11 +622,14 @@ namespace ElectraDecodersUtil
 				int32 GetWidth() const;
 				int32 GetHeight() const;
 				void GetDisplaySize(int32& OutWidth, int32& OutHeight) const;
+				void GetDisplaySize(uint32& OutWidth, uint32& OutHeight) const;
 				void GetCrop(int32& OutLeft, int32& OutRight, int32& OutTop, int32& OutBottom) const;
 				void GetAspect(int32& OutSarW, int32& OutSarH) const;
 				FFractionalValue GetTiming() const;
 				uint64 GetConstraintFlags() const;
 				FString GetRFC6381(const TCHAR* SampleTypePrefix) const;
+				FString GetFormatInfo() const;
+				static FString GetFormatInfo(uint32 InProfileIDC, uint32 InLevelIDC, uint64 InConstraintFlags);
 
 				uint8 sps_video_parameter_set_id;							// u(4)
 				uint8 sps_max_sub_layers_minus1;							// u(3), 0-6
@@ -1357,10 +1360,82 @@ namespace ElectraDecodersUtil
 
 			bool ELECTRADECODERS_API ParseVideoParameterSet(TMap<uint32, FVideoParameterSet>& InOutVideoParameterSets, const uint8* InBitstream, uint64 InBitstreamLenInBytes);
 			bool ELECTRADECODERS_API ParseSequenceParameterSet(TMap<uint32, FSequenceParameterSet>& InOutSequenceParameterSets, const uint8* InBitstream, uint64 InBitstreamLenInBytes);
-			bool ELECTRADECODERS_API ParseSequenceParameterSet(FSequenceParameterSet& OutSequenceParameterSet, const uint8* InBitstream, uint64 InBitstreamLenInBytes);
 			bool ELECTRADECODERS_API ParsePictureParameterSet(TMap<uint32, FPictureParameterSet>& InOutPictureParameterSets, const TMap<uint32, FSequenceParameterSet>& InSequenceParameterSets, const uint8* InBitstream, uint64 InBitstreamLenInBytes);
 			bool ELECTRADECODERS_API ParseSliceHeader(TUniquePtr<FRBSP>& OutRBSP, FBitstreamReader& OutRBSPReader, FSliceSegmentHeader& OutSlice, const TMap<uint32, FVideoParameterSet>& InVideoParameterSets, const TMap<uint32, FSequenceParameterSet>& InSequenceParameterSets, const TMap<uint32, FPictureParameterSet>& InPictureParameterSets, const uint8* InBitstream, uint64 InBitstreamLenInBytes);
 
+			bool ELECTRADECODERS_API ParseSequenceParameterSet(FSequenceParameterSet& OutSequenceParameterSet, const uint8* InBitstream, uint64 InBitstreamLenInBytes);
+			inline bool ParseSequenceParameterSet(FSequenceParameterSet& OutSequenceParameterSet, const TArray<uint8>& InData)
+			{ return ParseSequenceParameterSet(OutSequenceParameterSet, InData.GetData(), InData.Num()); }
+
+
+
+
+			class FHEVCDecoderConfigurationRecord
+			{
+			public:
+				ELECTRADECODERS_API bool Parse(const TConstArrayView<uint8>& InDCR);
+				ELECTRADECODERS_API bool CreateFromCodecSpecificData(const TConstArrayView<uint8>& InCSD);
+
+				inline const TArray<uint8>& GetRawData() const
+				{ return RawData; }
+				inline const TArray<uint8>& GetCodecSpecificData() const
+				{ return CSD; }
+
+				inline const TArray<TArray<uint8>>& GetSequenceParameterSets() const
+				{ return SequenceParameterSets; }
+				inline uint8 GetGeneralProfileSpace() const
+				{ return GeneralProfileSpace; }
+				inline uint8 GetGeneralTierFlag() const
+				{ return GeneralTierFlag; }
+				inline uint8 GetGeneralProfileIDC() const
+				{ return GeneralProfileIDC; }
+				inline uint32 GetGeneralProfileCompatibilityFlags() const
+				{ return GeneralProfileCompatibilityFlags; }
+				inline uint64 GetGeneralConstraintIndicatorFlags() const
+				{ return GeneralConstraintIndicatorFlags; }
+				inline uint8 GetGeneralLevelIDC() const
+				{ return GeneralLevelIDC; }
+				inline uint8 GetChromaFormat() const
+				{ return ChromaFormat; }
+				inline uint8 GetBitDepthLumaMinus8() const
+				{ return BitDepthLumaMinus8; }
+				inline uint8 GetBitDepthChromaMinus8() const
+				{ return BitDepthChromaMinus8; }
+
+				ELECTRADECODERS_API FString GetCodecSpecifierRFC6381(const TCHAR* InPrefix) const;
+				ELECTRADECODERS_API FString GetFormatInfo() const;
+
+			private:
+				struct FArray
+				{
+					TArray<TArray<uint8>> NALUs;
+					uint8 Completeness = 0;
+					uint8 NALUnitType = 0;
+				};
+				TArray<FArray> Arrays;
+				TArray<TArray<uint8>> SequenceParameterSets;
+				TArray<FSequenceParameterSet> ParsedSequenceParameterSets;
+				TArray<uint8> RawData;
+				TArray<uint8> CSD;
+				uint64 GeneralConstraintIndicatorFlags = 0;
+				uint32 GeneralProfileCompatibilityFlags = 0;
+				uint16 MinSpatialSegmentationIDC = 0;
+				uint16 AverageFrameRate = 0;
+				uint8 ConfigurationVersion = 0;
+				uint8 GeneralProfileSpace = 0;
+				uint8 GeneralTierFlag = 0;
+				uint8 GeneralProfileIDC = 0;
+				uint8 GeneralLevelIDC = 0;
+				uint8 ParallelismType = 0;
+				uint8 ChromaFormat = 0;
+				uint8 BitDepthLumaMinus8 = 0;
+				uint8 BitDepthChromaMinus8 = 0;
+				uint8 ConstantFrameRate = 0;
+				uint8 NumTemporalLayers = 0;
+				uint8 TemporalIdNested = 0;
+				uint8 NALUnitLengthMinus1 = 0;
+			};
+
 		} // namespace H265
 	} // namespace MPEG
 } // namespace ElectraDecodersUtil
diff --git a/Engine/Plugins/Media/ElectraPlayer/Source/ElectraPlayerRuntime/Private/Runtime/Decoder/VideoDecoderHelpers.cpp b/Engine/Plugins/Media/ElectraPlayer/Source/ElectraPlayerRuntime/Private/Runtime/Decoder/VideoDecoderHelpers.cpp
--- a/Engine/Plugins/Media/ElectraPlayer/Source/ElectraPlayerRuntime/Private/Runtime/Decoder/VideoDecoderHelpers.cpp
+++ b/Engine/Plugins/Media/ElectraPlayer/Source/ElectraPlayerRuntime/Private/Runtime/Decoder/VideoDecoderHelpers.cpp
@@ -1,6 +1,6 @@
 // Copyright Epic Games, Inc. All Rights Reserved.
 
-#include "VideoDecoderHelpers.h"
+#include "Decoder/VideoDecoderHelpers.h"
 
 #include "PlayerCore.h"
 #include "Utilities/UtilsMP4.h"
@@ -159,7 +159,7 @@ void FHDRHelper::Update(int32 BitDepth, const FColorimetryHelper& InColorimetry,
 		if (ActiveMasteringDisplayColourVolume.IsSet())
 		{
 			ElectraDecodersUtil::MPEG::FSEImastering_display_colour_volume seim;
-			if (seim.ParseFromMessage(ActiveMasteringDisplayColourVolume.GetValue()))
+			if (ElectraDecodersUtil::MPEG::ParseFromSEIMessage(seim, ActiveMasteringDisplayColourVolume.GetValue()))
 			{
 				const float ST2086_Chroma_Scale = 50000.0f;
 				const float ST2086_Luma_Scale = 10000.0f;
@@ -183,7 +183,7 @@ void FHDRHelper::Update(int32 BitDepth, const FColorimetryHelper& InColorimetry,
 		if (ActiveContentLightLevelInfo.IsSet())
 		{
 			ElectraDecodersUtil::MPEG::FSEIcontent_light_level_info seim;
-			if (seim.ParseFromMessage(ActiveContentLightLevelInfo.GetValue()))
+			if (ElectraDecodersUtil::MPEG::ParseFromSEIMessage(seim, ActiveContentLightLevelInfo.GetValue()))
 			{
 				FVideoDecoderHDRMetadata_content_light_level_info ll;
 				ll.max_content_light_level = seim.max_content_light_level;
@@ -196,7 +196,7 @@ void FHDRHelper::Update(int32 BitDepth, const FColorimetryHelper& InColorimetry,
 		if (ActiveAlternativeTransferCharacteristics.IsSet())
 		{
 			ElectraDecodersUtil::MPEG::FSEIalternative_transfer_characteristics seim;
-			if (seim.ParseFromMessage(ActiveAlternativeTransferCharacteristics.GetValue()))
+			if (ElectraDecodersUtil::MPEG::ParseFromSEIMessage(seim, ActiveAlternativeTransferCharacteristics.GetValue()))
 			{
 				CurrentAlternativeTransferCharacteristics = seim.preferred_transfer_characteristics;
 			}
diff --git a/Engine/Plugins/Media/ElectraPlayer/Source/ElectraPlayerRuntime/Private/Runtime/Decoder/VideoDecoderHelpers.h b/Engine/Plugins/Media/ElectraPlayer/Source/ElectraPlayerRuntime/Private/Runtime/Decoder/VideoDecoderHelpers.h
deleted file mode 100644
diff --git a/Engine/Plugins/Media/ElectraPlayer/Source/ElectraPlayerRuntime/Public/Decoder/VideoDecoderHelpers.h b/Engine/Plugins/Media/ElectraPlayer/Source/ElectraPlayerRuntime/Public/Decoder/VideoDecoderHelpers.h
new file mode 100644
--- /dev/null
+++ b/Engine/Plugins/Media/ElectraPlayer/Source/ElectraPlayerRuntime/Public/Decoder/VideoDecoderHelpers.h
@@ -0,0 +1,107 @@
+// Copyright Epic Games, Inc. All Rights Reserved.
+
+#pragma once
+
+#include "CoreTypes.h"
+#include "Templates/SharedPointer.h"
+#include "Misc/Optional.h"
+#include "Utils/MPEG/ElectraUtilsMPEGVideo.h"
+#include "MediaVideoDecoderOutput.h"
+#include "ParameterDictionary.h"
+
+namespace Electra
+{
+
+namespace MPEG
+{
+
+class FVideoDecoderTimecode : public IVideoDecoderTimecode
+{
+public:
+	virtual ~FVideoDecoderTimecode() = default;
+	virtual FMPEGDefinition const* GetMPEGDefinition() const override
+	{ return &DecoderTimecode; }
+	void Update(const IVideoDecoderTimecode::FMPEGDefinition& InTimecode)
+	{ DecoderTimecode = InTimecode; }
+private:
+	IVideoDecoderTimecode::FMPEGDefinition DecoderTimecode;
+};
+
+
+
+class FColorimetryHelper
+{
+public:
+	ELECTRAPLAYERRUNTIME_API void Reset();
+	ELECTRAPLAYERRUNTIME_API void Update(uint8 colour_primaries, uint8 transfer_characteristics, uint8 matrix_coeffs, uint8 video_full_range_flag, uint8 video_format);
+	ELECTRAPLAYERRUNTIME_API void Update(const TArray<uint8>& InFromCOLRBox);
+	ELECTRAPLAYERRUNTIME_API void UpdateParamDict(FParamDict& InOutDictionary);
+	ELECTRAPLAYERRUNTIME_API bool GetCurrentValues(uint8& colour_primaries, uint8& transfer_characteristics, uint8& matrix_coeffs) const;
+
+private:
+	class FVideoDecoderColorimetry : public IVideoDecoderColorimetry
+	{
+	public:
+		FVideoDecoderColorimetry(uint8 colour_primaries, uint8 transfer_characteristics, uint8 matrix_coeffs, uint8 video_full_range_flag, uint8 video_format)
+		{
+			Colorimetry.ColourPrimaries = colour_primaries;
+			Colorimetry.TransferCharacteristics = transfer_characteristics;
+			Colorimetry.MatrixCoefficients = matrix_coeffs;
+			Colorimetry.VideoFullRangeFlag = video_full_range_flag;
+			Colorimetry.VideoFormat = video_format;
+		}
+		virtual ~FVideoDecoderColorimetry() = default;
+		IVideoDecoderColorimetry::FMPEGDefinition const* GetMPEGDefinition() const override
+		{ return &Colorimetry; }
+		IVideoDecoderColorimetry::FMPEGDefinition Colorimetry;
+	};
+
+	TSharedPtr<FVideoDecoderColorimetry, ESPMode::ThreadSafe> CurrentColorimetry;
+};
+
+
+
+class FHDRHelper
+{
+public:
+	ELECTRAPLAYERRUNTIME_API void Reset();
+	ELECTRAPLAYERRUNTIME_API void Update(int32 BitDepth, const FColorimetryHelper& InColorimetry, const TArray<ElectraDecodersUtil::MPEG::FSEIMessage>& InGlobalPrefixSEIs, const TArray<ElectraDecodersUtil::MPEG::FSEIMessage>& InLocalPrefixSEIs, bool bIsNewCLVS);
+	ELECTRAPLAYERRUNTIME_API void UpdateFromMPEGBoxes(int32 BitDepth, const FColorimetryHelper& InColorimetry, const TArray<uint8>& InMDCVBox, const TArray<uint8>& InCLLIBox);
+	ELECTRAPLAYERRUNTIME_API void Update(int32 BitDepth, const FColorimetryHelper& InColorimetry, const TOptional<FVideoDecoderHDRMetadata_mastering_display_colour_volume>& InMDCV, const TOptional<FVideoDecoderHDRMetadata_content_light_level_info>& InCLLI);
+	ELECTRAPLAYERRUNTIME_API void UpdateParamDict(FParamDict& InOutDictionary);
+private:
+	class FVideoDecoderHDRInformation : public IVideoDecoderHDRInformation
+	{
+	public:
+		virtual ~FVideoDecoderHDRInformation() = default;
+		IVideoDecoderHDRInformation::EType GetHDRType() const override
+		{ return HDRType; }
+		const FVideoDecoderHDRMetadata_mastering_display_colour_volume* GetMasteringDisplayColourVolume() const override
+		{ return MasteringDisplayColourVolume.GetPtrOrNull(); }
+		const FVideoDecoderHDRMetadata_content_light_level_info* GetContentLightLevelInfo() const override
+		{ return ContentLightLevelInfo.GetPtrOrNull(); }
+		void SetHDRType(IVideoDecoderHDRInformation::EType InHDRType)
+		{ HDRType = InHDRType; }
+		void SetMasteringDisplayColourVolume(const FVideoDecoderHDRMetadata_mastering_display_colour_volume& In)
+		{ MasteringDisplayColourVolume = In; }
+		void SetContentLightLevelInfo(const FVideoDecoderHDRMetadata_content_light_level_info& In)
+		{ ContentLightLevelInfo = In; }
+	private:
+		IVideoDecoderHDRInformation::EType HDRType = IVideoDecoderHDRInformation::EType::Unknown;
+		TOptional<FVideoDecoderHDRMetadata_mastering_display_colour_volume> MasteringDisplayColourVolume;
+		TOptional<FVideoDecoderHDRMetadata_content_light_level_info> ContentLightLevelInfo;
+	};
+	void SetHDRType(int32 BitDepth, const FColorimetryHelper& InColorimetry);
+	TSharedPtr<FVideoDecoderHDRInformation, ESPMode::ThreadSafe> CurrentHDRInfo;
+	TOptional<ElectraDecodersUtil::MPEG::FSEIMessage> ActiveMasteringDisplayColourVolume;
+	TOptional<ElectraDecodersUtil::MPEG::FSEIMessage> ActiveContentLightLevelInfo;
+	TOptional<ElectraDecodersUtil::MPEG::FSEIMessage> ActiveAlternativeTransferCharacteristics;
+	bool bIsFirst = true;
+	int32 CurrentAlternativeTransferCharacteristics = -1;
+};
+
+
+
+} // namespace MPEG
+} // namespace Electra
+
