Electra: Refactoring utilities to allow for external use

#rnx

diff --git a/Engine/Plugins/Media/ElectraCodecs/Source/ElectraDecoders/Private/Utils/AudioChannelMapper.cpp b/Engine/Plugins/Media/ElectraCodecs/Source/ElectraDecoders/Private/Utils/AudioChannelMapper.cpp
new file mode 100644
--- /dev/null
+++ b/Engine/Plugins/Media/ElectraCodecs/Source/ElectraDecoders/Private/Utils/AudioChannelMapper.cpp
@@ -0,0 +1,455 @@
+// Copyright Epic Games, Inc. All Rights Reserved.
+
+#include "Utils/AudioChannelMapper.h"
+#include "ElectraDecodersUtils.h"
+
+namespace Electra
+{
+
+	namespace DirectlyCopyable
+	{
+		// This defines the order in which FMediaAudioResampler expects the source channels.
+		// See comments in MediaAudioResampler.cpp regarding the row order.
+		#define CP IElectraDecoderAudioOutput::EChannelPosition
+		static const CP _1[] = { CP::C };
+		static const CP _2[] = { CP::L, CP::R };
+		static const CP _3[] = { CP::L, CP::R, CP::C };
+		static const CP _4[] = { CP::L, CP::R, CP::Ls, CP::Rs };
+		static const CP _5[] = { CP::L, CP::R, CP::C, CP::Ls, CP::Rs };
+		static const CP _6[] = { CP::L, CP::R, CP::C, CP::LFE, CP::Ls, CP::Rs };
+		static const CP _7[] = { CP::L, CP::R, CP::Lsr, CP::LFE, CP::Rsr, CP::Ls, CP::Rs };
+		static const CP _8[] = { CP::L, CP::R, CP::C, CP::LFE, CP::Ls, CP::Rs, CP::Lsr, CP::Rsr };
+		static const CP * const ResamplerOrderMap[] = { _1, _2, _3, _4, _5, _6, _7, _8 };
+		#undef CP
+	}
+
+	FAudioChannelMapper::FAudioChannelMapper()
+	{
+	}
+
+	FAudioChannelMapper::~FAudioChannelMapper()
+	{
+	}
+
+
+	bool FAudioChannelMapper::Initialize(TSharedPtr<IElectraDecoderAudioOutput, ESPMode::ThreadSafe> InSampleBlock)
+	{
+		Reset();
+
+		check(InSampleBlock->GetSampleFormat() == IElectraDecoderAudioOutput::ESampleFormat::Int16 ||
+			  InSampleBlock->GetSampleFormat() == IElectraDecoderAudioOutput::ESampleFormat::Float);
+		if (InSampleBlock->GetSampleFormat() != IElectraDecoderAudioOutput::ESampleFormat::Int16 && InSampleBlock->GetSampleFormat() != IElectraDecoderAudioOutput::ESampleFormat::Float)
+		{
+			return false;
+		}
+
+		const int32 Stride = InSampleBlock->GetBytesPerFrame();
+		const int32 NumBytesPerSample = InSampleBlock->GetBytesPerSample();
+		BytesPerSample = NumBytesPerSample;
+
+		bCanCopyDirectly = CanBeCopiedDirectly(InSampleBlock);
+		if (bCanCopyDirectly)
+		{
+			// Even though we can copy the source over directly we still set up the target source layout
+			// for later use. GetNumTargetChannels() requires this.
+			for(int32 i=0; i<InSampleBlock->GetNumChannels(); ++i)
+			{
+				FTargetSource& ts = TargetSources.AddDefaulted_GetRef();
+				ts.ChannelPosition = InSampleBlock->GetChannelPosition(i);
+				ts.FirstOffset = BytesPerSample * i;
+				ts.Stride = Stride;
+			}
+		}
+		else
+		{
+			// Setup the mapping from the input channels as a mutable array.
+			TArray<FTargetSource> InputSources;
+			for(int32 i=0; i<InSampleBlock->GetNumChannels(); ++i)
+			{
+				FTargetSource& ts = InputSources.AddDefaulted_GetRef();
+				ts.ChannelPosition = InSampleBlock->GetChannelPosition(i);
+				ts.FirstOffset = NumBytesPerSample * i;
+				ts.Stride = Stride;
+			}
+			// Now remove all the inputs the resampler does not support.
+			for(int32 i=0; i<InputSources.Num(); ++i)
+			{
+				if (InputSources[i].ChannelPosition != IElectraDecoderAudioOutput::EChannelPosition::C &&
+					InputSources[i].ChannelPosition != IElectraDecoderAudioOutput::EChannelPosition::L &&
+					InputSources[i].ChannelPosition != IElectraDecoderAudioOutput::EChannelPosition::R &&
+					InputSources[i].ChannelPosition != IElectraDecoderAudioOutput::EChannelPosition::LFE &&
+					InputSources[i].ChannelPosition != IElectraDecoderAudioOutput::EChannelPosition::Ls &&
+					InputSources[i].ChannelPosition != IElectraDecoderAudioOutput::EChannelPosition::Rs &&
+					InputSources[i].ChannelPosition != IElectraDecoderAudioOutput::EChannelPosition::Lsr &&
+					InputSources[i].ChannelPosition != IElectraDecoderAudioOutput::EChannelPosition::Rsr)
+				{
+					InputSources.RemoveAt(i);
+					--i;
+				}
+			}
+
+			// Find the smallest resampler layout that contains all the required channels in any order.
+			int32 ResamplerChannelLayout = MatchesResamplerLayout(InputSources);
+			// At the very least the last 8 channel layout has to match since we removed all other channels.
+			// If no channels remain we can't do anything.
+			if (ResamplerChannelLayout == 0)
+			{
+				return false;
+			}
+
+			// Create a channel mapping matching the channel layout of the resampler.
+			// Channels that the resampler uses but are not present are filled with empty
+			// mappings that create silence for that channel.
+			CreateResamplerChannelMapping(ResamplerChannelLayout-1, InputSources);
+		}
+		return true;
+	}
+
+
+	bool FAudioChannelMapper::Initialize(int32 NumBytesPerSample, uint32 ChannelConfiguration)
+	{
+		Reset();
+		#define CP(ch) { IElectraDecoderAudioOutput::EChannelPosition::ch }
+		switch(ChannelConfiguration)
+		{
+			default:
+			{
+				break;
+			}
+			case 1:
+			{
+				FSourceLayout Layout[] = { CP(C) };
+				return Initialize(NumBytesPerSample, Layout);
+			}
+			case 2:
+			case 8:	// handle unspecified channel1+channel2 configuration as L+R
+			{
+				FSourceLayout Layout[] = { CP(L), CP(R) };
+				return Initialize(NumBytesPerSample, Layout);
+			}
+			case 3:
+			{
+				FSourceLayout Layout[] = { CP(C), CP(L), CP(R) };
+				return Initialize(NumBytesPerSample, Layout);
+			}
+			case 4:
+			{
+				FSourceLayout Layout[] = { CP(C), CP(L), CP(R), CP(Cs) };
+				return Initialize(NumBytesPerSample, Layout);
+			}
+			case 5:
+			{
+				FSourceLayout Layout[] = { CP(C), CP(L), CP(R), CP(Ls), CP(Rs) };
+				return Initialize(NumBytesPerSample, Layout);
+			}
+			case 6:
+			{
+				FSourceLayout Layout[] = { CP(C), CP(L), CP(R), CP(Ls), CP(Rs), CP(LFE) };
+				return Initialize(NumBytesPerSample, Layout);
+			}
+			case 7:
+			{
+				FSourceLayout Layout[] = { CP(C), CP(Lc), CP(Rc), CP(L), CP(R), CP(Ls), CP(Rs), CP(LFE) };
+				return Initialize(NumBytesPerSample, Layout);
+			}
+			case 9:
+			{
+				FSourceLayout Layout[] = { CP(L), CP(R), CP(Cs) };
+				return Initialize(NumBytesPerSample, Layout);
+			}
+			case 10:
+			{
+				FSourceLayout Layout[] = { CP(L), CP(R), CP(Ls), CP(Rs) };
+				return Initialize(NumBytesPerSample, Layout);
+			}
+			case 11:
+			{
+				FSourceLayout Layout[] = { CP(C), CP(L), CP(R), CP(Ls), CP(Rs), CP(Cs), CP(LFE) };
+				return Initialize(NumBytesPerSample, Layout);
+			}
+			case 12:
+			{
+				FSourceLayout Layout[] = { CP(C), CP(L), CP(R), CP(Ls), CP(Rs), CP(Lsr), CP(Rsr), CP(LFE) };
+				return Initialize(NumBytesPerSample, Layout);
+			}
+			case 13:
+			{
+				FSourceLayout Layout[] = { CP(C), CP(Lc), CP(Rc), CP(L), CP(R), CP(Lss), CP(Rss), CP(Lsr), CP(Rsr), CP(Cs), CP(LFE), CP(LFE2),
+										   CP(Cv), CP(Lv), CP(Rv), CP(Lvss), CP(Rvss), CP(Ts), CP(Lvr), CP(Rvr), CP(Cvr), CP(Cb), CP(Lb), CP(Rb) };
+				return Initialize(NumBytesPerSample, Layout);
+			}
+			case 14:
+			{
+				FSourceLayout Layout[] = { CP(C), CP(L), CP(R), CP(Ls), CP(Rs), CP(LFE), CP(Lv), CP(Rv) };
+				return Initialize(NumBytesPerSample, Layout);
+			}
+		}
+		#undef CP
+		return false;
+	}
+
+	bool FAudioChannelMapper::Initialize(int32 NumBytesPerSample, TArrayView<const FSourceLayout> ChannelPositions)
+	{
+		check(NumBytesPerSample == 2 || NumBytesPerSample == 4);
+		Reset();
+
+		bCanCopyDirectly = CanBeCopiedDirectly(ChannelPositions);
+		if (bCanCopyDirectly)
+		{
+			// Even though we can copy the source over directly we still set up the target source layout
+			// for later use. GetNumTargetChannels() requires this.
+			for(int32 i=0; i<ChannelPositions.Num(); ++i)
+			{
+				FTargetSource& ts = TargetSources.AddDefaulted_GetRef();
+				ts.ChannelPosition = ChannelPositions[i].ChannelPosition;
+				ts.FirstOffset = NumBytesPerSample * i;
+				ts.Stride = NumBytesPerSample * ChannelPositions.Num();
+			}
+		}
+		else
+		{
+			int32 Stride = NumBytesPerSample * ChannelPositions.Num();
+			// Setup the mapping from the input channels as a mutable array.
+			TArray<FTargetSource> InputSources;
+			for(int32 i=0; i<ChannelPositions.Num(); ++i)
+			{
+				FTargetSource& ts = InputSources.AddDefaulted_GetRef();
+				ts.ChannelPosition = ChannelPositions[i].ChannelPosition;
+				ts.FirstOffset = NumBytesPerSample * i;
+				ts.Stride = Stride;
+			}
+			// Now remove all the inputs we do not support.
+			for(int32 i=0; i<InputSources.Num(); ++i)
+			{
+				if (InputSources[i].ChannelPosition != IElectraDecoderAudioOutput::EChannelPosition::C &&
+					InputSources[i].ChannelPosition != IElectraDecoderAudioOutput::EChannelPosition::L &&
+					InputSources[i].ChannelPosition != IElectraDecoderAudioOutput::EChannelPosition::R &&
+					InputSources[i].ChannelPosition != IElectraDecoderAudioOutput::EChannelPosition::LFE &&
+					InputSources[i].ChannelPosition != IElectraDecoderAudioOutput::EChannelPosition::Ls &&
+					InputSources[i].ChannelPosition != IElectraDecoderAudioOutput::EChannelPosition::Rs &&
+					InputSources[i].ChannelPosition != IElectraDecoderAudioOutput::EChannelPosition::Lsr &&
+					InputSources[i].ChannelPosition != IElectraDecoderAudioOutput::EChannelPosition::Rsr)
+				{
+					InputSources.RemoveAt(i);
+					--i;
+				}
+			}
+
+			// Do the current sources match one of the channel layouts used by the resampler?
+			int32 ResamplerChannelLayout = MatchesResamplerLayout(InputSources);
+			// At the very least the last 8 channel layout has to match.
+			check(ResamplerChannelLayout);
+			if (ResamplerChannelLayout == 0)
+			{
+				return false;
+			}
+
+			// Create a channel mapping matching the channel layout of the resampler.
+			// Channels that the resampler uses but are not present are filled with empty
+			// mappings that create silence for that channel.
+			CreateResamplerChannelMapping(ResamplerChannelLayout-1, InputSources);
+		}
+
+		BytesPerSample = NumBytesPerSample;
+		return true;
+	}
+
+	bool FAudioChannelMapper::IsInitialized() const
+	{
+		return BytesPerSample != 0;
+	}
+
+
+	void FAudioChannelMapper::Reset()
+	{
+		TargetSources.Empty();
+		BytesPerSample = 0;
+		bCanCopyDirectly = false;
+	}
+
+	int32 FAudioChannelMapper::GetNumTargetChannels() const
+	{
+		return TargetSources.Num();
+	}
+
+	void FAudioChannelMapper::MapChannels(void* OutputBuffer, uint32 OutputBufferSize, const void* InputBuffer, uint32 InputBufferSize, int32 NumSamplesPerChannel) const
+	{
+		// Must have been initialized!
+		check(IsInitialized());
+		if (IsInitialized())
+		{
+			if (bCanCopyDirectly)
+			{
+				uint32 nc = InputBufferSize < OutputBufferSize ? InputBufferSize : OutputBufferSize;
+				FMemory::Memcpy(OutputBuffer, InputBuffer, nc);
+			}
+			else
+			{
+				int32 OutStride = BytesPerSample * TargetSources.Num();
+				// Copy channels in output order
+				for(int32 Ch=0,ChMax=TargetSources.Num(); Ch<ChMax; ++Ch)
+				{
+					void* OutCh = ElectraDecodersUtil::AdvancePointer(OutputBuffer, Ch * BytesPerSample);
+					const void* InCh = ElectraDecodersUtil::AdvancePointer(InputBuffer, TargetSources[Ch].FirstOffset);
+					int32 InStride = TargetSources[Ch].Stride;
+					if (BytesPerSample == 2)
+					{
+						if (TargetSources[Ch].ChannelPosition == IElectraDecoderAudioOutput::EChannelPosition::Disabled)
+						{
+							for(int32 i=0; i<NumSamplesPerChannel; ++i)
+							{
+								*reinterpret_cast<uint16*>(OutCh) = 0;
+								OutCh = ElectraDecodersUtil::AdvancePointer(OutCh, OutStride);
+							}
+						}
+						else
+						{
+							for(int32 i=0; i<NumSamplesPerChannel; ++i)
+							{
+								*reinterpret_cast<uint16*>(OutCh) = *reinterpret_cast<const uint16*>(InCh);
+								OutCh = ElectraDecodersUtil::AdvancePointer(OutCh, OutStride);
+								InCh = ElectraDecodersUtil::AdvancePointer(InCh, InStride);
+							}
+						}
+					}
+					else if (BytesPerSample == 4)
+					{
+						if (TargetSources[Ch].ChannelPosition == IElectraDecoderAudioOutput::EChannelPosition::Disabled)
+						{
+							for(int32 i=0; i<NumSamplesPerChannel; ++i)
+							{
+								*reinterpret_cast<uint32*>(OutCh) = 0;
+								OutCh = ElectraDecodersUtil::AdvancePointer(OutCh, OutStride);
+							}
+						}
+						else
+						{
+							for(int32 i=0; i<NumSamplesPerChannel; ++i)
+							{
+								*reinterpret_cast<uint32*>(OutCh) = *reinterpret_cast<const uint32*>(InCh);
+								OutCh = ElectraDecodersUtil::AdvancePointer(OutCh, OutStride);
+								InCh = ElectraDecodersUtil::AdvancePointer(InCh, InStride);
+							}
+						}
+					}
+				}
+			}
+		}
+		else
+		{
+			FMemory::Memzero(OutputBuffer, OutputBufferSize);
+		}
+	}
+
+
+	bool FAudioChannelMapper::CanBeCopiedDirectly(TSharedPtr<IElectraDecoderAudioOutput, ESPMode::ThreadSafe> InSampleBlock)
+	{
+		int32 nc = InSampleBlock->GetNumChannels();
+		int32 NumUnspecifiedChannels = 0;
+		for(int32 i=0; i<nc; ++i)
+		{
+			IElectraDecoderAudioOutput::EChannelPosition cp = InSampleBlock->GetChannelPosition(i);
+			if (cp >= IElectraDecoderAudioOutput::EChannelPosition::Unspec0 && cp <= IElectraDecoderAudioOutput::EChannelPosition::Unspec31)
+			{
+				++NumUnspecifiedChannels;
+			}
+		}
+		// If all channel positions are unspecified then we can copy directly. There are no known
+		// positions, so it could be anything really.
+		if (NumUnspecifiedChannels == nc)
+		{
+			return true;
+		}
+
+		if (nc >= 1 && nc <= 8)
+		{
+			const IElectraDecoderAudioOutput::EChannelPosition * const ResamplerOrder = DirectlyCopyable::ResamplerOrderMap[nc - 1];
+			if (ResamplerOrder)
+			{
+				for(int32 i=0; i<nc; ++i)
+				{
+					if (InSampleBlock->GetChannelPosition(i) != ResamplerOrder[i])
+					{
+						return false;
+					}
+				}
+				return true;
+			}
+		}
+		return false;
+	}
+
+	bool FAudioChannelMapper::CanBeCopiedDirectly(TArrayView<const FSourceLayout> ChannelPositions)
+	{
+		if (ChannelPositions.Num() >= 1 && ChannelPositions.Num() <= 8)
+		{
+			const IElectraDecoderAudioOutput::EChannelPosition * const ResamplerOrder = DirectlyCopyable::ResamplerOrderMap[ChannelPositions.Num() - 1];
+			if (ResamplerOrder)
+			{
+				for(int32 i=0; i<ChannelPositions.Num(); ++i)
+				{
+					if (ResamplerOrder[i] != ChannelPositions[i].ChannelPosition)
+					{
+						return false;
+					}
+				}
+				return true;
+			}
+		}
+		return false;
+	}
+
+	int32 FAudioChannelMapper::MatchesResamplerLayout(TArray<FTargetSource>& InputSources)
+	{
+		if (InputSources.Num())
+		{
+			for(int32 i=InputSources.Num(); i<=UE_ARRAY_COUNT(DirectlyCopyable::ResamplerOrderMap); ++i)
+			{
+				TArrayView<const IElectraDecoderAudioOutput::EChannelPosition> ResamplerChannels(DirectlyCopyable::ResamplerOrderMap[i-1], i);
+				bool bAllChannelsPresent = true;
+				for(int32 j=0; j<InputSources.Num(); ++j)
+				{
+					if (!ResamplerChannels.Contains(InputSources[j].ChannelPosition))
+					{
+						bAllChannelsPresent = false;
+						break;
+					}
+				}
+				if (bAllChannelsPresent)
+				{
+					return i;
+				}
+			}
+		}
+		return 0;
+	}
+
+	void FAudioChannelMapper::CreateResamplerChannelMapping(int32 ResamplerChannelIndex, const TArray<FTargetSource>& InSources)
+	{
+		check(ResamplerChannelIndex >= 0 && ResamplerChannelIndex < UE_ARRAY_COUNT(DirectlyCopyable::ResamplerOrderMap));
+		if (ResamplerChannelIndex >= 0 && ResamplerChannelIndex < UE_ARRAY_COUNT(DirectlyCopyable::ResamplerOrderMap))
+		{
+			TArrayView<const IElectraDecoderAudioOutput::EChannelPosition> ResamplerChannels(DirectlyCopyable::ResamplerOrderMap[ResamplerChannelIndex], ResamplerChannelIndex+1);
+			for(auto &ResamplerChannel : ResamplerChannels)
+			{
+				FTargetSource ts;
+				int32 SourceIndex = InSources.IndexOfByPredicate([ResamplerChannel](const FTargetSource& In){return In.ChannelPosition == ResamplerChannel;});
+				if (SourceIndex != INDEX_NONE)
+				{
+					TargetSources.Emplace(InSources[SourceIndex]);
+				}
+				else
+				{
+					TargetSources.AddDefaulted();
+				}
+			}
+		}
+		else
+		{
+			// Add at least one empty entry to keep everything happy.
+			TargetSources.AddDefaulted();
+		}
+	}
+
+
+} // namespace Electra
diff --git a/Engine/Plugins/Media/ElectraCodecs/Source/ElectraDecoders/Public/Utils/AudioChannelMapper.h b/Engine/Plugins/Media/ElectraCodecs/Source/ElectraDecoders/Public/Utils/AudioChannelMapper.h
new file mode 100644
--- /dev/null
+++ b/Engine/Plugins/Media/ElectraCodecs/Source/ElectraDecoders/Public/Utils/AudioChannelMapper.h
@@ -0,0 +1,79 @@
+// Copyright Epic Games, Inc. All Rights Reserved.
+#pragma once
+
+#include "CoreTypes.h"
+#include "Containers/Array.h"
+#include "Containers/ArrayView.h"
+#include "Templates/SharedPointer.h"
+
+#include "IElectraDecoderOutputAudio.h"
+
+namespace Electra
+{
+	/**
+	 * Class to map audio channels from some decoder-specific output order into the
+	 * order expected by FMediaAudioResampler.
+	 *
+	 * Channels that are not supported by UE will be ignored and are silent.
+	 * No effort to mix them into other channels is made.
+	 * As fas as media playback is concerned we handle the most common channel layouts
+	 * only, like mono, stereo, 5.1 and 7.1.
+	 */
+	class ELECTRADECODERS_API FAudioChannelMapper
+	{
+	public:
+/*
+		// Channel position as per ISO/IEC 23001-8 Table 7.
+		enum class IElectraDecoderAudioOutput::EChannelPosition
+		{
+			L, R, C, LFE, Ls, Rs, Lc, Rc,
+			Lsr, Rsr, Cs, Lsd, Rsd, Lss, Rss, Lw,
+			Rw, Lv, Rv, Cv, Lvr, Rvr, Cvr, Lvss,
+			Rvss, Ts, LFE2, Lb, Rb, Cb, Lvs, Rvs,
+			Disabled = 255
+		};
+*/
+		struct FSourceLayout
+		{
+			IElectraDecoderAudioOutput::EChannelPosition ChannelPosition = IElectraDecoderAudioOutput::EChannelPosition::Disabled;
+		};
+
+		FAudioChannelMapper();
+		~FAudioChannelMapper();
+
+		// Initializes with an ISO/IEC 23001-8 channel configuration value from Table 7 (compatible with ISO/IEC 14496-3:2009/Amd 4)
+		bool Initialize(int32 NumBytesPerSample, uint32 ChannelConfiguration);
+
+		// Initializes with ISO/IEC 23001-8 channel positions.
+		bool Initialize(int32 NumBytesPerSample, TArrayView<const FSourceLayout> ChannelPositions);
+
+		bool IsInitialized() const;
+
+		void Reset();
+
+		// Gets the number of output channels mapped to the UE internal layout.
+		int32 GetNumTargetChannels() const;
+
+		// Maps the supported input channels to the output channels in UE internal layout.
+		void MapChannels(void* OutputBuffer, uint32 OutputBufferSize, const void* InputBuffer, uint32 InputBufferSize, int32 NumSamplesPerChannel) const;
+
+		bool Initialize(TSharedPtr<IElectraDecoderAudioOutput, ESPMode::ThreadSafe> InSampleBlock);
+
+	private:
+		struct FTargetSource
+		{
+			IElectraDecoderAudioOutput::EChannelPosition ChannelPosition = IElectraDecoderAudioOutput::EChannelPosition::Disabled;
+			int32 FirstOffset = 0;
+			int32 Stride = 0;
+		};
+
+		static bool CanBeCopiedDirectly(TArrayView<const FSourceLayout> ChannelPositions);
+		bool CanBeCopiedDirectly(TSharedPtr<IElectraDecoderAudioOutput, ESPMode::ThreadSafe> InSampleBlock);
+		int32 MatchesResamplerLayout(TArray<FTargetSource>& InputSources);
+		void CreateResamplerChannelMapping(int32 ResamplerChannelIndex, const TArray<FTargetSource>& InSources);
+
+		TArray<FTargetSource> TargetSources;
+		int32 BytesPerSample = 0;
+		bool bCanCopyDirectly = false;
+	};
+} // namespace Electra
diff --git a/Engine/Plugins/Media/ElectraPlayer/Source/ElectraPlayerRuntime/Private/ElectraPlayerDecoderResourceManager.cpp b/Engine/Plugins/Media/ElectraPlayer/Source/ElectraPlayerRuntime/Private/ElectraPlayerDecoderResourceManager.cpp
new file mode 100644
--- /dev/null
+++ b/Engine/Plugins/Media/ElectraPlayer/Source/ElectraPlayerRuntime/Private/ElectraPlayerDecoderResourceManager.cpp
@@ -0,0 +1,27 @@
+// Copyright Epic Games, Inc. All Rights Reserved.
+
+#include "IElectraPlayerDecoderResourceManager.h"
+#include "Decoder/ElectraDecoderResourceManager.h"
+#include "ElectraPlayerPlatform.h"
+#include "Renderer/RendererVideo.h"
+
+
+TSharedPtr<IElectraDecoderResourceDelegate, ESPMode::ThreadSafe> FElectraPlayerDecoderResourceManager::GetDelegate()
+{
+	return Electra::FPlatformElectraDecoderResourceManager::GetDelegate();
+}
+
+TSharedPtr<Electra::IVideoDecoderResourceDelegate, ESPMode::ThreadSafe> FElectraPlayerDecoderResourceManager::CreatePlatformVideoDecoderResourceDelegate(TSharedPtr<IElectraPlayerAdapterDelegate, ESPMode::ThreadSafe> InElectraPlayerAdapterDelegate)
+{
+	return Electra::PlatformCreateVideoDecoderResourceDelegate(InElectraPlayerAdapterDelegate);
+}
+
+bool FElectraPlayerDecoderResourceManager::SetupRenderBufferFromDecoderOutput(Electra::IMediaRenderer::IBuffer* InOutBufferToSetup, TSharedPtr<Electra::FParamDict, ESPMode::ThreadSafe> InOutBufferPropertes, TSharedPtr<IElectraDecoderVideoOutput, ESPMode::ThreadSafe> InDecoderOutput, IElectraDecoderResourceDelegateBase::IDecoderPlatformResource* InPlatformSpecificResource)
+{
+	return Electra::FPlatformElectraDecoderResourceManager::SetupRenderBufferFromDecoderOutput(InOutBufferToSetup, InOutBufferPropertes, InDecoderOutput, InPlatformSpecificResource);
+}
+
+FVideoDecoderOutput* FElectraPlayerDecoderResourceManager::FVideo::Create()
+{
+	return FElectraPlayerPlatformVideoDecoderOutputFactory::Create();
+}
diff --git a/Engine/Plugins/Media/ElectraPlayer/Source/ElectraPlayerRuntime/Private/Runtime/Decoder/AudioDecoder.cpp b/Engine/Plugins/Media/ElectraPlayer/Source/ElectraPlayerRuntime/Private/Runtime/Decoder/AudioDecoder.cpp
--- a/Engine/Plugins/Media/ElectraPlayer/Source/ElectraPlayerRuntime/Private/Runtime/Decoder/AudioDecoder.cpp
+++ b/Engine/Plugins/Media/ElectraPlayer/Source/ElectraPlayerRuntime/Private/Runtime/Decoder/AudioDecoder.cpp
@@ -15,8 +15,8 @@
 #include "Player/PlayerSessionServices.h"
 #include "Utilities/Utilities.h"
 #include "Utilities/StringHelpers.h"
-#include "Utilities/AudioChannelMapper.h"
 #include "Utils/MPEG/ElectraUtilsMPEGAudio.h"
+#include "Utils/AudioChannelMapper.h"
 
 // Error codes must be in the 1000-1999 range. 1-999 is reserved for the decoder implementation.
 #define ERRCODE_AUDIO_INTERNAL_COULD_NOT_CREATE_DECODER			1001
diff --git a/Engine/Plugins/Media/ElectraPlayer/Source/ElectraPlayerRuntime/Private/Runtime/Renderer/RendererBase.h b/Engine/Plugins/Media/ElectraPlayer/Source/ElectraPlayerRuntime/Private/Runtime/Renderer/RendererBase.h
--- a/Engine/Plugins/Media/ElectraPlayer/Source/ElectraPlayerRuntime/Private/Runtime/Renderer/RendererBase.h
+++ b/Engine/Plugins/Media/ElectraPlayer/Source/ElectraPlayerRuntime/Private/Runtime/Renderer/RendererBase.h
@@ -2,255 +2,4 @@
 
 #pragma once
 
-#include "PlayerCore.h"
-#include "PlayerTime.h"
-#include "ParameterDictionary.h"
-
-#include "MediaDecoderOutput.h"
-
-namespace Electra
-{
-
-	class IMediaRenderClock
-	{
-	public:
-		virtual ~IMediaRenderClock() = default;
-
-		enum class ERendererType
-		{
-			Video,
-			Audio,
-			Subtitles,
-		};
-
-		/**
-		 * Called by the renderer to set the time of the most sample that has been output most recently.
-		 *
-		 * @param ForRenderer
-		 *               Identifies the type of renderer from which the time is set.
-		 * @param CurrentRenderTime
-		 *               The time of the sample most recently output.
-		 */
-		virtual void SetCurrentTime(ERendererType ForRenderer, const FTimeValue& CurrentRenderTime) = 0;
-
-		/**
-		 * Gets the current _interpolated_ sample output time from the last
-		 * time SetCurrentTime() was called plus the elapsed time since then.
-		 *
-		 * @param FromRenderer
-		 *               Identifies the type of renderer for which to get the time.
-		 *
-		 * @return Interpolated render time
-		 */
-		virtual FTimeValue GetInterpolatedRenderTime(ERendererType FromRenderer) = 0;
-	};
-
-
-
-
-
-	class IMediaRenderer : public IDecoderOutputOwner
-	{
-	public:
-		virtual ~IMediaRenderer() = default;
-
-		/**
-		 * Buffer interface
-		 */
-		class IBuffer
-		{
-		public:
-			virtual ~IBuffer() = default;
-
-			/**
-			 * Returns the properties of the buffer.
-			 *
-			 * @return Const reference to the dictionary containing buffer properties.
-			 */
-			virtual const FParamDict& GetBufferProperties() const = 0;
-			
-			// Same, but returns the writable dictionary. Use with caution.
-			virtual FParamDict& GetMutableBufferProperties() = 0;
-		};
-
-
-		//=================================================================================================================
-		// Methods called from both a decoder and the player from their respective threads.
-		//
-
-		/**
-		 * Returns the properties of the buffer pool. Those properties should not change
-		 * during the lifetime of the pool.
-		 * A mandatory property is the maximum number of buffers that can be obtained from
-		 * the pool ("max_buffers":int64).
-		 *
-		 * @return Const reference to the dictionary containing buffer pool properties.
-		 */
-		virtual const FParamDict& GetBufferPoolProperties() const = 0;
-
-
-
-		//=================================================================================================================
-		// Methods called from a decoder (and from within the decoder thread)
-		//
-
-		/**
-		 * Create a buffer pool from where a decoder can get the block of memory to decode into.
-		 *
-		 * @param Parameters Dictionary with create options depending on the type of
-		 *                   render buffers to create (video or audio)
-		 *
-		 * @return One of UEMEDIA_ERROR_OK, UEMEDIA_ERROR_BAD_ARGUMENTS or UEMEDIA_ERROR_OOM (or others if warranted).
-		 */
-		virtual UEMediaError CreateBufferPool(const FParamDict& Parameters) = 0;
-
-		/**
-		 * Asks for a sample buffer from the buffer pool created previously through CreateBufferPool().
-		 *
-		 * @param OutBuffer Receives the address of the returned IBuffer
-		 * @param TimeoutInMicroseconds
-		 *                  Maximum time to wait for a buffer to become available before giving up.
-		 *                  If 0 a "peek" if a buffer is available is to be performed and a buffer
-		 *                  returned if one is available.
-		 *                  No value indicates an infinite wait time!
-		 * @param InParameters
-		 *                  Optional parameters for buffer properties if different from those to CreateBufferPool().
-		 *
-		 * @return One of UEMEDIA_ERROR_OK, UEMEDIA_ERROR_BAD_ARGUMENTS or UEMEDIA_ERROR_INSUFFICIENT_DATA.
-		 *         Other values if warranted but those will cause a playback error.
-		 */
-		virtual UEMediaError AcquireBuffer(IBuffer*& OutBuffer, int32 TimeoutInMicroseconds, const FParamDict& InParameters) = 0;
-
-		/**
-		 * Releases the buffer for rendering and subsequent return to the buffer pool.
-		 *
-		 * @param Buffer  The buffer being returned to the renderer for rendering.
-		 * @param bRender true to render the contents of the buffer, false to discard them.
-		 * @param InSampleProperties
-		 *                Properties of the decoded sample (eg. width and height for video, sample rate for audio)
-		 *
-		 * @return Should be UEMEDIA_ERROR_OK only. Anything else will cause a playback error (like UEMEDIA_ERROR_BAD_ARGUMENTS for bad properties)
-		 */
-		virtual UEMediaError ReturnBuffer(IBuffer* Buffer, bool bRender, FParamDict& InOutSampleProperties) = 0;
-
-		/**
-		 * Informs that the decoder is done with this pool. This only indicates no calls to
-		 * GetBuffer() and ReleaseBuffer() will be made by this decoder instance any more.
-		 * It is up to the renderer to decide whether or not to destroy the pool or keep it
-		 * around for re-use with a new decoder instance (the next CreateBufferPool() call).
-		 *
-		 * @return Should be UEMEDIA_ERROR_OK only. Anything else will cause a playback error.
-		 */
-		virtual UEMediaError ReleaseBufferPool() = 0;
-
-
-		/**
-		 * Returns if the renderer's output queues are ready to receive any output from a decode
-		 *
-		 * @return True if ready, false if decode should be delayed.
-		 */
-		virtual bool CanReceiveOutputFrames(uint64 NumFrames) const = 0;
-
-		/**
-		 * Returns how many frames of ready to use data the renderer has in its queue.
-		 * 
-		 * @return True if information is available, false if not.
-		 */
-		virtual bool GetEnqueuedFrameInfo(int32& OutNumberOfEnqueuedFrames, FTimeValue& OutDurationOfEnqueuedFrames) const = 0;
-
-		//=================================================================================================================
-		// Methods called from the player (and from within the player thread)
-		//
-
-		/**
-		 * Sets the render clock which this renderer needs to update with the presentation time
-		 * of the most recently rendered sample.
-		 *
-		 * @param RenderClock
-		 *               Render clock to update with the most recent output sample time.
-		 */
-		virtual void SetRenderClock(TSharedPtr<IMediaRenderClock, ESPMode::ThreadSafe> RenderClock) = 0;
-
-
-		/**
-		 * Called if this renderer is being wrapped by another renderer.
-		 */
-		virtual void SetParentRenderer(TWeakPtr<IMediaRenderer, ESPMode::ThreadSafe> ParentRenderer) = 0;
-
-		/**
-		 * Sets the next expected sample's approximate presentation time stamp.
-		 * This is called as data buffering begins at the indicated play start
-		 * position. The actual timestamp of the first sample will vary.
-		 * The renderer can ignore this value and rely only on the timestamp
-		 * of the decoded sample.
-		 *
-		 * @param NextApproxPTS
-		 *               Approximate timestamp at which playback is about to commence.
-		 */
-		virtual void SetNextApproximatePresentationTime(const FTimeValue& NextApproxPTS) = 0;
-
-		/**
-		 * Flushes all pending buffers not yet rendered.
-		 * There should be no outstanding buffers the decoder has not yet released.
-		 * Options may indicate to hold the current frame on-screen (for video) or
-		 * to tear it down.
-		 *
-		 * @param InOptions Options indicating flush operation, like keeping the current video frame on-screen.
-		 *
-		 * @return Should be UEMEDIA_ERROR_OK only.
-		 *         If there are still frames a decoder has not returned yet UEMEDIA_ERROR_INTERNAL should be returned.
-		 *         Any failure will cause a playback error.
-		 */
-		virtual UEMediaError Flush(const FParamDict& InOptions) = 0;
-
-		/**
-		 * Begins rendering of the first sample buffer.
-		 *
-		 * @param InOptions Reserved for future use.
-		 */
-		virtual void StartRendering(const FParamDict& InOptions) = 0;
-
-		/**
-		 * Stops rendering of sample buffers.
-		 *
-		 * @param InOptions Reserved for future use.
-		 */
-		virtual void StopRendering(const FParamDict& InOptions) = 0;
-
-		/**
-		 * Work any regular tasks on the output buffer pool
-		 *
-		 * @note Expected to be called as needed from (normally) the decoder worker thread
-		 */
-		virtual void TickOutputBufferPool() {};
-	};
-
-
-	namespace RenderOptionKeys
-	{
-		static const FName ValidityValue(TEXT("$renderVV$"));
-		static const FName Duration(TEXT("duration"));
-		static const FName PTS(TEXT("pts"));
-		static const FName EOSFlag(TEXT("eos"));
-		static const FName DummyBufferFlag(TEXT("is_dummy"));
-
-		static const FName Texture(TEXT("texture"));
-		static const FName AllocatedAddress(TEXT("address"));
-		static const FName AllocatedSize(TEXT("size"));
-
-		static const FName MaxBuffers(TEXT("max_buffers"));
-		static const FName MaxBufferSize(TEXT("max_buffer_size"));
-		static const FName NumBuffers(TEXT("num_buffers"));
-		static const FName UsedByteSize(TEXT("byte_size"));
-
-		static const FName MaxChannels(TEXT("max_channels"));
-		static const FName NumChannels(TEXT("num_channels"));
-		static const FName SamplesPerBlock(TEXT("samples_per_block"));
-		static const FName SampleRate(TEXT("sample_rate"));
-	}
-
-} // namespace Electra
-
-
-
-
+#include "MediaRendererBase.h"
diff --git a/Engine/Plugins/Media/ElectraPlayer/Source/ElectraPlayerRuntime/Private/Runtime/Utilities/AudioChannelMapper.cpp b/Engine/Plugins/Media/ElectraPlayer/Source/ElectraPlayerRuntime/Private/Runtime/Utilities/AudioChannelMapper.cpp
deleted file mode 100644
diff --git a/Engine/Plugins/Media/ElectraPlayer/Source/ElectraPlayerRuntime/Private/Runtime/Utilities/AudioChannelMapper.h b/Engine/Plugins/Media/ElectraPlayer/Source/ElectraPlayerRuntime/Private/Runtime/Utilities/AudioChannelMapper.h
deleted file mode 100644
diff --git a/Engine/Plugins/Media/ElectraPlayer/Source/ElectraPlayerRuntime/Public/IElectraPlayerDecoderResourceManager.h b/Engine/Plugins/Media/ElectraPlayer/Source/ElectraPlayerRuntime/Public/IElectraPlayerDecoderResourceManager.h
new file mode 100644
--- /dev/null
+++ b/Engine/Plugins/Media/ElectraPlayer/Source/ElectraPlayerRuntime/Public/IElectraPlayerDecoderResourceManager.h
@@ -0,0 +1,30 @@
+// Copyright Epic Games, Inc. All Rights Reserved.
+
+#pragma once
+
+#include "CoreTypes.h"
+#include "IElectraDecoderResourceDelegate.h"
+#include "MediaRendererBase.h"
+#include "ParameterDictionary.h"
+#include "MediaVideoDecoderOutput.h"
+
+class IElectraPlayerAdapterDelegate;
+class IElectraDecoderVideoOutput;
+
+namespace Electra
+{
+class IVideoDecoderResourceDelegate;
+}
+
+class FElectraPlayerDecoderResourceManager
+{
+public:
+	ELECTRAPLAYERRUNTIME_API static TSharedPtr<IElectraDecoderResourceDelegate, ESPMode::ThreadSafe> GetDelegate();
+	ELECTRAPLAYERRUNTIME_API static TSharedPtr<Electra::IVideoDecoderResourceDelegate, ESPMode::ThreadSafe> CreatePlatformVideoDecoderResourceDelegate(TSharedPtr<IElectraPlayerAdapterDelegate, ESPMode::ThreadSafe> InElectraPlayerAdapterDelegate);
+	ELECTRAPLAYERRUNTIME_API static bool SetupRenderBufferFromDecoderOutput(Electra::IMediaRenderer::IBuffer* InOutBufferToSetup, TSharedPtr<Electra::FParamDict, ESPMode::ThreadSafe> InOutBufferPropertes, TSharedPtr<IElectraDecoderVideoOutput, ESPMode::ThreadSafe> InDecoderOutput, IElectraDecoderResourceDelegateBase::IDecoderPlatformResource* InPlatformSpecificResource);
+	class FVideo
+	{
+	public:
+		ELECTRAPLAYERRUNTIME_API static FVideoDecoderOutput* Create();
+	};
+};
diff --git a/Engine/Plugins/Media/ElectraUtil/Source/ElectraBase/Public/MediaDecoderOutput.h b/Engine/Plugins/Media/ElectraUtil/Source/ElectraBase/Public/MediaDecoderOutput.h
--- a/Engine/Plugins/Media/ElectraUtil/Source/ElectraBase/Public/MediaDecoderOutput.h
+++ b/Engine/Plugins/Media/ElectraUtil/Source/ElectraBase/Public/MediaDecoderOutput.h
@@ -356,4 +356,5 @@ static const FName FPSNumerator(TEXT("fps_num"));
 static const FName FPSDenominator(TEXT("fps_denom"));
 static const FName PixelDataScale(TEXT("pix_datascale"));
 static const FName Timecode(TEXT("timecode"));
+static const FName TMCDTimecode(TEXT("tmcd_timecode"));
 }
diff --git a/Engine/Plugins/Media/ElectraUtil/Source/ElectraBase/Public/MediaRendererBase.h b/Engine/Plugins/Media/ElectraUtil/Source/ElectraBase/Public/MediaRendererBase.h
new file mode 100644
--- /dev/null
+++ b/Engine/Plugins/Media/ElectraUtil/Source/ElectraBase/Public/MediaRendererBase.h
@@ -0,0 +1,251 @@
+// Copyright Epic Games, Inc. All Rights Reserved.
+
+#pragma once
+
+#include "CoreTypes.h"
+#include "Core/MediaTypes.h"
+#include "PlayerTime.h"
+#include "ParameterDictionary.h"
+#include "MediaDecoderOutput.h"
+
+namespace Electra
+{
+	class IMediaRenderClock
+	{
+	public:
+		virtual ~IMediaRenderClock() = default;
+
+		enum class ERendererType
+		{
+			Video,
+			Audio,
+			Subtitles,
+		};
+
+		/**
+		 * Called by the renderer to set the time of the most sample that has been output most recently.
+		 *
+		 * @param ForRenderer
+		 *               Identifies the type of renderer from which the time is set.
+		 * @param CurrentRenderTime
+		 *               The time of the sample most recently output.
+		 */
+		virtual void SetCurrentTime(ERendererType ForRenderer, const FTimeValue& CurrentRenderTime) = 0;
+
+		/**
+		 * Gets the current _interpolated_ sample output time from the last
+		 * time SetCurrentTime() was called plus the elapsed time since then.
+		 *
+		 * @param FromRenderer
+		 *               Identifies the type of renderer for which to get the time.
+		 *
+		 * @return Interpolated render time
+		 */
+		virtual FTimeValue GetInterpolatedRenderTime(ERendererType FromRenderer) = 0;
+	};
+
+
+
+
+
+	class IMediaRenderer : public IDecoderOutputOwner
+	{
+	public:
+		virtual ~IMediaRenderer() = default;
+
+		/**
+		 * Buffer interface
+		 */
+		class IBuffer
+		{
+		public:
+			virtual ~IBuffer() = default;
+
+			/**
+			 * Returns the properties of the buffer.
+			 *
+			 * @return Const reference to the dictionary containing buffer properties.
+			 */
+			virtual const FParamDict& GetBufferProperties() const = 0;
+
+			// Same, but returns the writable dictionary. Use with caution.
+			virtual FParamDict& GetMutableBufferProperties() = 0;
+		};
+
+
+		//=================================================================================================================
+		// Methods called from both a decoder and the player from their respective threads.
+		//
+
+		/**
+		 * Returns the properties of the buffer pool. Those properties should not change
+		 * during the lifetime of the pool.
+		 * A mandatory property is the maximum number of buffers that can be obtained from
+		 * the pool ("max_buffers":int64).
+		 *
+		 * @return Const reference to the dictionary containing buffer pool properties.
+		 */
+		virtual const FParamDict& GetBufferPoolProperties() const = 0;
+
+
+
+		//=================================================================================================================
+		// Methods called from a decoder (and from within the decoder thread)
+		//
+
+		/**
+		 * Create a buffer pool from where a decoder can get the block of memory to decode into.
+		 *
+		 * @param Parameters Dictionary with create options depending on the type of
+		 *                   render buffers to create (video or audio)
+		 *
+		 * @return One of UEMEDIA_ERROR_OK, UEMEDIA_ERROR_BAD_ARGUMENTS or UEMEDIA_ERROR_OOM (or others if warranted).
+		 */
+		virtual UEMediaError CreateBufferPool(const FParamDict& Parameters) = 0;
+
+		/**
+		 * Asks for a sample buffer from the buffer pool created previously through CreateBufferPool().
+		 *
+		 * @param OutBuffer Receives the address of the returned IBuffer
+		 * @param TimeoutInMicroseconds
+		 *                  Maximum time to wait for a buffer to become available before giving up.
+		 *                  If 0 a "peek" if a buffer is available is to be performed and a buffer
+		 *                  returned if one is available.
+		 *                  No value indicates an infinite wait time!
+		 * @param InParameters
+		 *                  Optional parameters for buffer properties if different from those to CreateBufferPool().
+		 *
+		 * @return One of UEMEDIA_ERROR_OK, UEMEDIA_ERROR_BAD_ARGUMENTS or UEMEDIA_ERROR_INSUFFICIENT_DATA.
+		 *         Other values if warranted but those will cause a playback error.
+		 */
+		virtual UEMediaError AcquireBuffer(IBuffer*& OutBuffer, int32 TimeoutInMicroseconds, const FParamDict& InParameters) = 0;
+
+		/**
+		 * Releases the buffer for rendering and subsequent return to the buffer pool.
+		 *
+		 * @param Buffer  The buffer being returned to the renderer for rendering.
+		 * @param bRender true to render the contents of the buffer, false to discard them.
+		 * @param InSampleProperties
+		 *                Properties of the decoded sample (eg. width and height for video, sample rate for audio)
+		 *
+		 * @return Should be UEMEDIA_ERROR_OK only. Anything else will cause a playback error (like UEMEDIA_ERROR_BAD_ARGUMENTS for bad properties)
+		 */
+		virtual UEMediaError ReturnBuffer(IBuffer* Buffer, bool bRender, FParamDict& InOutSampleProperties) = 0;
+
+		/**
+		 * Informs that the decoder is done with this pool. This only indicates no calls to
+		 * GetBuffer() and ReleaseBuffer() will be made by this decoder instance any more.
+		 * It is up to the renderer to decide whether or not to destroy the pool or keep it
+		 * around for re-use with a new decoder instance (the next CreateBufferPool() call).
+		 *
+		 * @return Should be UEMEDIA_ERROR_OK only. Anything else will cause a playback error.
+		 */
+		virtual UEMediaError ReleaseBufferPool() = 0;
+
+
+		/**
+		 * Returns if the renderer's output queues are ready to receive any output from a decode
+		 *
+		 * @return True if ready, false if decode should be delayed.
+		 */
+		virtual bool CanReceiveOutputFrames(uint64 NumFrames) const = 0;
+
+		/**
+		 * Returns how many frames of ready to use data the renderer has in its queue.
+		 *
+		 * @return True if information is available, false if not.
+		 */
+		virtual bool GetEnqueuedFrameInfo(int32& OutNumberOfEnqueuedFrames, FTimeValue& OutDurationOfEnqueuedFrames) const = 0;
+
+		//=================================================================================================================
+		// Methods called from the player (and from within the player thread)
+		//
+
+		/**
+		 * Sets the render clock which this renderer needs to update with the presentation time
+		 * of the most recently rendered sample.
+		 *
+		 * @param RenderClock
+		 *               Render clock to update with the most recent output sample time.
+		 */
+		virtual void SetRenderClock(TSharedPtr<IMediaRenderClock, ESPMode::ThreadSafe> RenderClock) = 0;
+
+
+		/**
+		 * Called if this renderer is being wrapped by another renderer.
+		 */
+		virtual void SetParentRenderer(TWeakPtr<IMediaRenderer, ESPMode::ThreadSafe> ParentRenderer) = 0;
+
+		/**
+		 * Sets the next expected sample's approximate presentation time stamp.
+		 * This is called as data buffering begins at the indicated play start
+		 * position. The actual timestamp of the first sample will vary.
+		 * The renderer can ignore this value and rely only on the timestamp
+		 * of the decoded sample.
+		 *
+		 * @param NextApproxPTS
+		 *               Approximate timestamp at which playback is about to commence.
+		 */
+		virtual void SetNextApproximatePresentationTime(const FTimeValue& NextApproxPTS) = 0;
+
+		/**
+		 * Flushes all pending buffers not yet rendered.
+		 * There should be no outstanding buffers the decoder has not yet released.
+		 * Options may indicate to hold the current frame on-screen (for video) or
+		 * to tear it down.
+		 *
+		 * @param InOptions Options indicating flush operation, like keeping the current video frame on-screen.
+		 *
+		 * @return Should be UEMEDIA_ERROR_OK only.
+		 *         If there are still frames a decoder has not returned yet UEMEDIA_ERROR_INTERNAL should be returned.
+		 *         Any failure will cause a playback error.
+		 */
+		virtual UEMediaError Flush(const FParamDict& InOptions) = 0;
+
+		/**
+		 * Begins rendering of the first sample buffer.
+		 *
+		 * @param InOptions Reserved for future use.
+		 */
+		virtual void StartRendering(const FParamDict& InOptions) = 0;
+
+		/**
+		 * Stops rendering of sample buffers.
+		 *
+		 * @param InOptions Reserved for future use.
+		 */
+		virtual void StopRendering(const FParamDict& InOptions) = 0;
+
+		/**
+		 * Work any regular tasks on the output buffer pool
+		 *
+		 * @note Expected to be called as needed from (normally) the decoder worker thread
+		 */
+		virtual void TickOutputBufferPool() {};
+	};
+
+
+	namespace RenderOptionKeys
+	{
+		static const FName ValidityValue(TEXT("$renderVV$"));
+		static const FName Duration(TEXT("duration"));
+		static const FName PTS(TEXT("pts"));
+		static const FName EOSFlag(TEXT("eos"));
+		static const FName DummyBufferFlag(TEXT("is_dummy"));
+
+		static const FName Texture(TEXT("texture"));
+		static const FName AllocatedAddress(TEXT("address"));
+		static const FName AllocatedSize(TEXT("size"));
+
+		static const FName MaxBuffers(TEXT("max_buffers"));
+		static const FName MaxBufferSize(TEXT("max_buffer_size"));
+		static const FName NumBuffers(TEXT("num_buffers"));
+		static const FName UsedByteSize(TEXT("byte_size"));
+
+		static const FName MaxChannels(TEXT("max_channels"));
+		static const FName NumChannels(TEXT("num_channels"));
+		static const FName SamplesPerBlock(TEXT("samples_per_block"));
+		static const FName SampleRate(TEXT("sample_rate"));
+	}
+
+} // namespace Electra
diff --git a/Engine/Plugins/Media/ElectraUtil/Source/ElectraSamples/Private/ElectraAudioSample.cpp b/Engine/Plugins/Media/ElectraUtil/Source/ElectraSamples/Private/ElectraAudioSample.cpp
new file mode 100644
--- /dev/null
+++ b/Engine/Plugins/Media/ElectraUtil/Source/ElectraSamples/Private/ElectraAudioSample.cpp
@@ -0,0 +1,43 @@
+// Copyright Epic Games, Inc. All Rights Reserved.
+
+#include "ElectraAudioSample.h"
+
+FElectraAudioSample::~FElectraAudioSample()
+{
+	FMemory::Free(Buffer);
+}
+
+bool FElectraAudioSample::Initialize(const void* InData, EMediaAudioSampleFormat InFormat, uint32 InNumChannels, uint32 InNumFrames, uint32 InSampleRate, const FMediaTimeStamp& InTime, const FTimespan& InDuration)
+{
+	check(InFormat == EMediaAudioSampleFormat::Float);
+
+	uint32 NumBytesNeeded = sizeof(float) * InNumChannels * InNumFrames;
+	if (NumBytesNeeded > NumBytesAllocated)
+	{
+		Buffer = FMemory::Realloc(Buffer, NumBytesNeeded);
+		if (!Buffer)
+		{
+			NumBytesAllocated = 0;
+			return false;
+		}
+		NumBytesAllocated = NumBytesNeeded;
+	}
+	MediaAudioSampleFormat = InFormat;
+	NumChannels = InNumChannels;
+	NumFrames = InNumFrames;
+	SampleRate = InSampleRate;
+	MediaTimeStamp = InTime;
+	Duration = InDuration;
+	if (Buffer && InData)
+	{
+		FMemory::Memcpy(Buffer, InData, NumBytesNeeded);
+	}
+	return true;
+}
+
+#if !UE_SERVER
+void FElectraAudioSample::ShutdownPoolable()
+{
+	// Nothing to do here at the moment.
+}
+#endif
diff --git a/Engine/Plugins/Media/ElectraUtil/Source/ElectraSamples/Private/ElectraAudioSample.h b/Engine/Plugins/Media/ElectraUtil/Source/ElectraSamples/Private/ElectraAudioSample.h
new file mode 100644
--- /dev/null
+++ b/Engine/Plugins/Media/ElectraUtil/Source/ElectraSamples/Private/ElectraAudioSample.h
@@ -0,0 +1,4 @@
+// Copyright Epic Games, Inc. All Rights Reserved.
+#pragma once
+
+#include "IElectraAudioSample.h"
diff --git a/Engine/Plugins/Media/ElectraUtil/Source/ElectraSamples/Public/IElectraAudioSample.h b/Engine/Plugins/Media/ElectraUtil/Source/ElectraSamples/Public/IElectraAudioSample.h
new file mode 100644
--- /dev/null
+++ b/Engine/Plugins/Media/ElectraUtil/Source/ElectraSamples/Public/IElectraAudioSample.h
@@ -0,0 +1,60 @@
+// Copyright Epic Games, Inc. All Rights Reserved.
+#pragma once
+
+#include "IMediaAudioSample.h"
+#include "Templates/SharedPointer.h"
+#include "MediaObjectPool.h"
+
+
+class ELECTRASAMPLES_API FElectraAudioSample
+	: public IMediaAudioSample
+	, public IMediaPoolable
+{
+public:
+	virtual ~FElectraAudioSample();
+	bool Initialize(const void* InData, EMediaAudioSampleFormat InFormat, uint32 InNumChannels, uint32 InNumFrames, uint32 InSampleRate, const FMediaTimeStamp& InTime, const FTimespan& InDuration);
+	uint32 GetAllocatedSize() const
+	{ return NumBytesAllocated;}
+
+	// From IMediaAudioSample
+	const void* GetBuffer() override
+	{ return Buffer; }
+	uint32 GetChannels() const override
+	{ return NumChannels; }
+	FTimespan GetDuration() const override
+	{ return Duration; }
+	EMediaAudioSampleFormat GetFormat() const override
+	{ return MediaAudioSampleFormat; }
+	uint32 GetFrames() const override
+	{ return NumFrames; }
+	uint32 GetSampleRate() const override
+	{ return SampleRate; }
+	FMediaTimeStamp GetTime() const override
+	{ return MediaTimeStamp; }
+	//TOptional<FTimecode> GetTimecode() const override { return TOptional<FTimecode>(); }
+
+	// From IMediaPoolable
+#if !UE_SERVER
+	void ShutdownPoolable() override;
+#endif
+
+protected:
+	FMediaTimeStamp MediaTimeStamp;
+	FTimespan Duration;
+	EMediaAudioSampleFormat MediaAudioSampleFormat = EMediaAudioSampleFormat::Float;
+	void* Buffer = nullptr;
+	uint32 NumChannels = 0;
+	uint32 NumFrames = 0;
+	uint32 SampleRate = 0;
+	uint32 NumBytesAllocated = 0;
+};
+
+
+using FElectraAudioSamplePtr = TSharedPtr<FElectraAudioSample, ESPMode::ThreadSafe>;
+using FElectraAudioSampleRef = TSharedRef<FElectraAudioSample, ESPMode::ThreadSafe>;
+
+class FElectraAudioSamplePool : public TMediaObjectPool<FElectraAudioSample>
+{
+public:
+	~FElectraAudioSamplePool() = default;
+};
