ElectraDecoders: Added decoder specific bitstream processors

#rnx
#donotcheckin

diff --git a/Engine/Plugins/Media/AppleProResDecoderElectra/Source/AppleProResDecoderElectra/Private/ProResDecoder/Mac/ElectraMediaProResDecoder.cpp b/Engine/Plugins/Media/AppleProResDecoderElectra/Source/AppleProResDecoderElectra/Private/ProResDecoder/Mac/ElectraMediaProResDecoder.cpp
--- a/Engine/Plugins/Media/AppleProResDecoderElectra/Source/AppleProResDecoderElectra/Private/ProResDecoder/Mac/ElectraMediaProResDecoder.cpp
+++ b/Engine/Plugins/Media/AppleProResDecoderElectra/Source/AppleProResDecoderElectra/Private/ProResDecoder/Mac/ElectraMediaProResDecoder.cpp
@@ -15,6 +15,7 @@
 #include "IElectraDecoderResourceDelegate.h"
 #include "ElectraDecodersUtils.h"
 #include "Utils/MPEG/ElectraUtilsMP4.h"	// to read the ProRes header, which is given in the form of an MP4 Atom
+#include "Utils/ElectraBitstreamProcessorGenericVideo.h"
 #include "AppleProResDecoderElectraModule.h"
 #include "Apple/AppleElectraDecoderPlatformOutputFormatTypes.h"
 
@@ -200,6 +201,13 @@ public:
 	EOutputStatus HaveOutput() override;
 	TSharedPtr<IElectraDecoderOutput, ESPMode::ThreadSafe> GetOutput() override;
 
+	TSharedPtr<IElectraDecoderBitstreamProcessor, ESPMode::ThreadSafe> CreateBitstreamProcessor() override
+	{
+		TMap<FString, FVariant> DecoderFeatures;
+		GetFeatures(DecoderFeatures);
+		return FElectraDecoderBitstreamProcessorGenericVideo::Create(DecoderFeatures, InitialCreationOptions);
+	}
+
 	void Suspend() override;
 	void Resume() override;
 
diff --git a/Engine/Plugins/Media/AppleProResDecoderElectra/Source/AppleProResDecoderElectra/Private/ProResDecoder/Windows/ElectraMediaProResDecoder.cpp b/Engine/Plugins/Media/AppleProResDecoderElectra/Source/AppleProResDecoderElectra/Private/ProResDecoder/Windows/ElectraMediaProResDecoder.cpp
--- a/Engine/Plugins/Media/AppleProResDecoderElectra/Source/AppleProResDecoderElectra/Private/ProResDecoder/Windows/ElectraMediaProResDecoder.cpp
+++ b/Engine/Plugins/Media/AppleProResDecoderElectra/Source/AppleProResDecoderElectra/Private/ProResDecoder/Windows/ElectraMediaProResDecoder.cpp
@@ -15,6 +15,7 @@
 #include "IElectraDecoderResourceDelegate.h"
 #include "ElectraDecodersUtils.h"
 #include "Utils/MPEG/ElectraUtilsMP4.h"	// to read the ProRes header, which is given in the form of an MP4 Atom
+#include "Utils/ElectraBitstreamProcessorGenericVideo.h"
 #include COMPILED_PLATFORM_HEADER(ElectraDecoderGPUBufferHelpers.h)
 
 
@@ -231,6 +232,13 @@ public:
 	EOutputStatus HaveOutput() override;
 	TSharedPtr<IElectraDecoderOutput, ESPMode::ThreadSafe> GetOutput() override;
 
+	TSharedPtr<IElectraDecoderBitstreamProcessor, ESPMode::ThreadSafe> CreateBitstreamProcessor() override
+	{
+		TMap<FString, FVariant> DecoderFeatures;
+		GetFeatures(DecoderFeatures);
+		return FElectraDecoderBitstreamProcessorGenericVideo::Create(DecoderFeatures, InitialCreationOptions);
+	}
+
 	void Suspend() override
 	{ }
 	void Resume() override
@@ -368,6 +376,7 @@ private:
 	uint32 Codec4CC = 0;
 	uint32 BufferAllocationSize = 0;
 
+	TMap<FString, FVariant> InitialCreationOptions;
 	TWeakPtr<IElectraDecoderResourceDelegate, ESPMode::ThreadSafe> ResourceDelegate;
 
 	IElectraDecoder::FError LastError;
@@ -509,6 +518,8 @@ void FElectraMediaProResDecoder::Shutdown()
 /*********************************************************************************************************************/
 
 FVideoDecoderProResElectra::FVideoDecoderProResElectra(const TMap<FString, FVariant>& InOptions, TSharedPtr<IElectraDecoderResourceDelegate, ESPMode::ThreadSafe> InResourceDelegate)
+	: InitialCreationOptions(InOptions)
+	, ResourceDelegate(InResourceDelegate)
 {
 	AspectW = (int32)ElectraDecodersUtil::GetVariantValueSafeI64(InOptions, TEXT("aspect_w"), 0);
 	AspectH = (int32)ElectraDecodersUtil::GetVariantValueSafeI64(InOptions, TEXT("aspect_h"), 0);
@@ -518,7 +529,6 @@ FVideoDecoderProResElectra::FVideoDecoderProResElectra(const TMap<FString, FVari
 	// The decoded width is the same as the display width
 	DecodedWidth = DisplayWidth;
 	DecodedHeight = DisplayHeight;
-	ResourceDelegate = InResourceDelegate;
 
 	MaxOutputBuffers = (uint32)ElectraDecodersUtil::GetVariantValueSafeU64(InOptions, TEXT("max_output_buffers"), 5);
 	MaxOutputBuffers += kElectraDecoderPipelineExtraFrames;
diff --git a/Engine/Plugins/Media/AvidDNxHDDecoderElectra/Source/Private/Windows/ElectraMediaAvidDNxHDDecoder.cpp b/Engine/Plugins/Media/AvidDNxHDDecoderElectra/Source/Private/Windows/ElectraMediaAvidDNxHDDecoder.cpp
--- a/Engine/Plugins/Media/AvidDNxHDDecoderElectra/Source/Private/Windows/ElectraMediaAvidDNxHDDecoder.cpp
+++ b/Engine/Plugins/Media/AvidDNxHDDecoderElectra/Source/Private/Windows/ElectraMediaAvidDNxHDDecoder.cpp
@@ -18,6 +18,7 @@
 #include "IElectraDecoderOutputVideo.h"
 #include "IElectraDecoderResourceDelegate.h"
 #include "ElectraDecodersUtils.h"
+#include "Utils/ElectraBitstreamProcessorDefault.h"
 #include COMPILED_PLATFORM_HEADER(ElectraDecoderGPUBufferHelpers.h)
 
 #include "AvidDNxCodec.h"
@@ -252,6 +253,9 @@ public:
 	EOutputStatus HaveOutput() override;
 	TSharedPtr<IElectraDecoderOutput, ESPMode::ThreadSafe> GetOutput() override;
 
+	TSharedPtr<IElectraDecoderBitstreamProcessor, ESPMode::ThreadSafe> CreateBitstreamProcessor() override
+	{ return FElectraDecoderBitstreamProcessorDefault::Create(); }
+
 	void Suspend() override
 	{ }
 	void Resume() override
diff --git a/Engine/Plugins/Media/D3D12VideoDecodersElectra/Source/Private/VideoDecoder_D3D12_Common.h b/Engine/Plugins/Media/D3D12VideoDecodersElectra/Source/Private/VideoDecoder_D3D12_Common.h
--- a/Engine/Plugins/Media/D3D12VideoDecodersElectra/Source/Private/VideoDecoder_D3D12_Common.h
+++ b/Engine/Plugins/Media/D3D12VideoDecodersElectra/Source/Private/VideoDecoder_D3D12_Common.h
@@ -13,6 +13,7 @@
 #include "IElectraDecoderOutputVideo.h"
 
 #include "ElectraDecodersUtils.h"
+#include "Utils/ElectraBitstreamProcessorDefault.h"
 
 /*********************************************************************************************************************/
 #include COMPILED_PLATFORM_HEADER(PlatformHeaders_Video_D3D.h)
@@ -357,6 +358,8 @@ protected:
 	EDecoderError Flush() = 0;
 	EOutputStatus HaveOutput() override;
 	TSharedPtr<IElectraDecoderOutput, ESPMode::ThreadSafe> GetOutput() override;
+	TSharedPtr<IElectraDecoderBitstreamProcessor, ESPMode::ThreadSafe> CreateBitstreamProcessor() override
+	{ return FElectraDecoderBitstreamProcessorDefault::Create(); }
 	void Suspend() override
 	{ }
 	void Resume() override
diff --git a/Engine/Plugins/Media/ElectraCodecs/Source/ElectraDecoders/Private/Android/aac/AudioDecoderAAC_Android.cpp b/Engine/Plugins/Media/ElectraCodecs/Source/ElectraDecoders/Private/Android/aac/AudioDecoderAAC_Android.cpp
--- a/Engine/Plugins/Media/ElectraCodecs/Source/ElectraDecoders/Private/Android/aac/AudioDecoderAAC_Android.cpp
+++ b/Engine/Plugins/Media/ElectraCodecs/Source/ElectraDecoders/Private/Android/aac/AudioDecoderAAC_Android.cpp
@@ -4,6 +4,7 @@
 #include "DecoderErrors_Android.h"
 #include "ElectraDecodersUtils.h"
 #include "Utils/MPEG/ElectraUtilsMPEGAudio.h"
+#include "Utils/ElectraBitstreamProcessorDefault.h"
 
 #include "IElectraDecoderFeaturesAndOptions.h"
 #include "IElectraDecoderOutputAudio.h"
@@ -119,6 +120,9 @@ public:
 	EOutputStatus HaveOutput() override;
 	TSharedPtr<IElectraDecoderOutput, ESPMode::ThreadSafe> GetOutput() override;
 
+	TSharedPtr<IElectraDecoderBitstreamProcessor, ESPMode::ThreadSafe> CreateBitstreamProcessor() override
+	{ return FElectraDecoderBitstreamProcessorDefault::Create(); }
+
 	void Suspend() override;
 	void Resume() override;
 
@@ -278,7 +282,7 @@ bool FElectraAudioDecoderAAC_Android::ResetToCleanStart()
 	bIsOutputFormatInfoValid = false;
 	LastPushedPresentationTimeUs = 0;
 	bDidEnqueueInputBuffer = false;
-	
+
 	return !LastError.IsSet();
 }
 
diff --git a/Engine/Plugins/Media/ElectraCodecs/Source/ElectraDecoders/Private/Android/h264/VideoDecoderH264_Android.cpp b/Engine/Plugins/Media/ElectraCodecs/Source/ElectraDecoders/Private/Android/h264/VideoDecoderH264_Android.cpp
--- a/Engine/Plugins/Media/ElectraCodecs/Source/ElectraDecoders/Private/Android/h264/VideoDecoderH264_Android.cpp
+++ b/Engine/Plugins/Media/ElectraCodecs/Source/ElectraDecoders/Private/Android/h264/VideoDecoderH264_Android.cpp
@@ -4,6 +4,7 @@
 #include "DecoderErrors_Android.h"
 #include "ElectraDecodersUtils.h"
 #include "Utils/MPEG/ElectraUtilsMPEGVideo.h"
+#include "Utils/MPEG/ElectraBitstreamProcessor_H264.h"
 
 #include "IElectraDecoderFeaturesAndOptions.h"
 #include "IElectraDecoderOutputVideo.h"
@@ -134,6 +135,13 @@ public:
 	EOutputStatus HaveOutput() override;
 	TSharedPtr<IElectraDecoderOutput, ESPMode::ThreadSafe> GetOutput() override;
 
+	TSharedPtr<IElectraDecoderBitstreamProcessor, ESPMode::ThreadSafe> CreateBitstreamProcessor() override
+	{
+		TMap<FString, FVariant> DecoderFeatures;
+		GetFeatures(DecoderFeatures);
+		return FElectraDecoderBitstreamProcessorH264::Create(DecoderFeatures, InitialCreationOptions);
+	}
+
 	void Suspend() override;
 	void Resume() override;
 
@@ -226,6 +234,7 @@ private:
 
 	IElectraDecoder::FError LastError;
 	FInitialMaxValues InitialMaxValues;
+	TMap<FString, FVariant> InitialCreationOptions;
 
 	TWeakPtr<IElectraDecoderResourceDelegate, ESPMode::ThreadSafe> ResourceDelegate;
 	IElectraDecoderResourceDelegateAndroid::IDecoderPlatformResourceAndroid* PlatformResource = nullptr;
@@ -276,6 +285,7 @@ TSharedPtr<IElectraDecoder, ESPMode::ThreadSafe> IElectraVideoDecoderH264_Androi
 
 
 FElectraVideoDecoderH264_Android::FElectraVideoDecoderH264_Android(const TMap<FString, FVariant>& InOptions, TSharedPtr<IElectraDecoderResourceDelegate, ESPMode::ThreadSafe> InResourceDelegate)
+	: InitialCreationOptions(InOptions)
 {
 	ResourceDelegate = InResourceDelegate;
 	PlatformResource = reinterpret_cast<IElectraDecoderResourceDelegateAndroid::IDecoderPlatformResourceAndroid*>(ElectraDecodersUtil::GetVariantValueSafeU64(InOptions, TEXT("platformresource"), 0));
diff --git a/Engine/Plugins/Media/ElectraCodecs/Source/ElectraDecoders/Private/Android/h265/VideoDecoderH265_Android.cpp b/Engine/Plugins/Media/ElectraCodecs/Source/ElectraDecoders/Private/Android/h265/VideoDecoderH265_Android.cpp
--- a/Engine/Plugins/Media/ElectraCodecs/Source/ElectraDecoders/Private/Android/h265/VideoDecoderH265_Android.cpp
+++ b/Engine/Plugins/Media/ElectraCodecs/Source/ElectraDecoders/Private/Android/h265/VideoDecoderH265_Android.cpp
@@ -4,6 +4,7 @@
 #include "DecoderErrors_Android.h"
 #include "ElectraDecodersUtils.h"
 #include "Utils/MPEG/ElectraUtilsMPEGVideo.h"
+#include "Utils/MPEG/ElectraBitstreamProcessor_H265.h"
 
 #include "IElectraDecoderFeaturesAndOptions.h"
 #include "IElectraDecoderOutputVideo.h"
@@ -134,6 +135,13 @@ public:
 	EOutputStatus HaveOutput() override;
 	TSharedPtr<IElectraDecoderOutput, ESPMode::ThreadSafe> GetOutput() override;
 
+	TSharedPtr<IElectraDecoderBitstreamProcessor, ESPMode::ThreadSafe> CreateBitstreamProcessor() override
+	{
+		TMap<FString, FVariant> DecoderFeatures;
+		GetFeatures(DecoderFeatures);
+		return FElectraDecoderBitstreamProcessorH265::Create(DecoderFeatures, InitialCreationOptions);
+	}
+
 	void Suspend() override;
 	void Resume() override;
 
@@ -226,6 +234,7 @@ private:
 
 	IElectraDecoder::FError LastError;
 	FInitialMaxValues InitialMaxValues;
+	TMap<FString, FVariant> InitialCreationOptions;
 
 	TWeakPtr<IElectraDecoderResourceDelegate, ESPMode::ThreadSafe> ResourceDelegate;
 	IElectraDecoderResourceDelegateAndroid::IDecoderPlatformResourceAndroid* PlatformResource = nullptr;
@@ -303,6 +312,7 @@ TSharedPtr<IElectraDecoder, ESPMode::ThreadSafe> IElectraVideoDecoderH265_Androi
 
 
 FElectraVideoDecoderH265_Android::FElectraVideoDecoderH265_Android(const TMap<FString, FVariant>& InOptions, TSharedPtr<IElectraDecoderResourceDelegate, ESPMode::ThreadSafe> InResourceDelegate)
+	: InitialCreationOptions(InOptions)
 {
 	ResourceDelegate = InResourceDelegate;
 	PlatformResource = reinterpret_cast<IElectraDecoderResourceDelegateAndroid::IDecoderPlatformResourceAndroid*>(ElectraDecodersUtil::GetVariantValueSafeU64(InOptions, TEXT("platformresource"), 0));
diff --git a/Engine/Plugins/Media/ElectraCodecs/Source/ElectraDecoders/Private/Android/vpx/VideoDecoderVPx_Android.cpp b/Engine/Plugins/Media/ElectraCodecs/Source/ElectraDecoders/Private/Android/vpx/VideoDecoderVPx_Android.cpp
--- a/Engine/Plugins/Media/ElectraCodecs/Source/ElectraDecoders/Private/Android/vpx/VideoDecoderVPx_Android.cpp
+++ b/Engine/Plugins/Media/ElectraCodecs/Source/ElectraDecoders/Private/Android/vpx/VideoDecoderVPx_Android.cpp
@@ -4,6 +4,7 @@
 #include "DecoderErrors_Android.h"
 #include "ElectraDecodersUtils.h"
 #include "Utils/Google/ElectraUtilsVPxVideo.h"
+#include "Utils/Google/ElectraBitstreamProcessorVPx.h"
 
 #include "IElectraDecoderFeaturesAndOptions.h"
 #include "IElectraDecoderOutputVideo.h"
@@ -134,6 +135,13 @@ public:
 	EOutputStatus HaveOutput() override;
 	TSharedPtr<IElectraDecoderOutput, ESPMode::ThreadSafe> GetOutput() override;
 
+	TSharedPtr<IElectraDecoderBitstreamProcessor, ESPMode::ThreadSafe> CreateBitstreamProcessor() override
+	{
+		TMap<FString, FVariant> DecoderFeatures;
+		GetFeatures(DecoderFeatures);
+		return FElectraDecoderBitstreamProcessorVPx::Create(DecoderFeatures, InitialCreationOptions);
+	}
+
 	void Suspend() override;
 	void Resume() override;
 
@@ -232,6 +240,7 @@ private:
 	FInitialMaxValues InitialMaxValues;
 	ECodec ActiveCodec = ECodec::VP9;
 
+	TMap<FString, FVariant> InitialCreationOptions;
 	TWeakPtr<IElectraDecoderResourceDelegate, ESPMode::ThreadSafe> ResourceDelegate;
 	IElectraDecoderResourceDelegateAndroid::IDecoderPlatformResourceAndroid* PlatformResource = nullptr;
 
@@ -306,8 +315,9 @@ TSharedPtr<IElectraDecoder, ESPMode::ThreadSafe> IElectraVideoDecoderVPx_Android
 
 
 FElectraVideoDecoderVPx_Android::FElectraVideoDecoderVPx_Android(const TMap<FString, FVariant>& InOptions, TSharedPtr<IElectraDecoderResourceDelegate, ESPMode::ThreadSafe> InResourceDelegate)
+	: InitialCreationOptions(InOptions)
+	, ResourceDelegate(InResourceDelegate)
 {
-	ResourceDelegate = InResourceDelegate;
 	PlatformResource = reinterpret_cast<IElectraDecoderResourceDelegateAndroid::IDecoderPlatformResourceAndroid*>(ElectraDecodersUtil::GetVariantValueSafeU64(InOptions, TEXT("platformresource"), 0));
 
 	uint32 Codec4CC = (uint32) ElectraDecodersUtil::GetVariantValueSafeI64(InOptions, TEXT("codec_4cc"), 0);
diff --git a/Engine/Plugins/Media/ElectraCodecs/Source/ElectraDecoders/Private/Apple/aac/AudioDecoderAAC_Apple.cpp b/Engine/Plugins/Media/ElectraCodecs/Source/ElectraDecoders/Private/Apple/aac/AudioDecoderAAC_Apple.cpp
--- a/Engine/Plugins/Media/ElectraCodecs/Source/ElectraDecoders/Private/Apple/aac/AudioDecoderAAC_Apple.cpp
+++ b/Engine/Plugins/Media/ElectraCodecs/Source/ElectraDecoders/Private/Apple/aac/AudioDecoderAAC_Apple.cpp
@@ -7,6 +7,7 @@
 #include "DecoderErrors_Apple.h"
 #include "ElectraDecodersUtils.h"
 #include "Utils/MPEG/ElectraUtilsMPEGAudio.h"
+#include "Utils/ElectraBitstreamProcessorDefault.h"
 
 #include "IElectraDecoderFeaturesAndOptions.h"
 #include "IElectraDecoderOutputAudio.h"
@@ -112,6 +113,9 @@ public:
 	EOutputStatus HaveOutput() override;
 	TSharedPtr<IElectraDecoderOutput, ESPMode::ThreadSafe> GetOutput() override;
 
+	TSharedPtr<IElectraDecoderBitstreamProcessor, ESPMode::ThreadSafe> CreateBitstreamProcessor() override
+	{ return FElectraDecoderBitstreamProcessorDefault::Create(); }
+
 	void Suspend() override;
 	void Resume() override;
 
diff --git a/Engine/Plugins/Media/ElectraCodecs/Source/ElectraDecoders/Private/Apple/h264/VideoDecoderH264_Apple.cpp b/Engine/Plugins/Media/ElectraCodecs/Source/ElectraDecoders/Private/Apple/h264/VideoDecoderH264_Apple.cpp
--- a/Engine/Plugins/Media/ElectraCodecs/Source/ElectraDecoders/Private/Apple/h264/VideoDecoderH264_Apple.cpp
+++ b/Engine/Plugins/Media/ElectraCodecs/Source/ElectraDecoders/Private/Apple/h264/VideoDecoderH264_Apple.cpp
@@ -7,6 +7,7 @@
 #include "DecoderErrors_Apple.h"
 #include "ElectraDecodersUtils.h"
 #include "Utils/MPEG/ElectraUtilsMPEGVideo.h"
+#include "Utils/MPEG/ElectraBitstreamProcessor_H264.h"
 
 #include "IElectraDecoderFeaturesAndOptions.h"
 #include "IElectraDecoderOutputVideo.h"
@@ -176,6 +177,13 @@ public:
 	EOutputStatus HaveOutput() override;
 	TSharedPtr<IElectraDecoderOutput, ESPMode::ThreadSafe> GetOutput() override;
 
+	TSharedPtr<IElectraDecoderBitstreamProcessor, ESPMode::ThreadSafe> CreateBitstreamProcessor() override
+	{
+		TMap<FString, FVariant> DecoderFeatures;
+		GetFeatures(DecoderFeatures);
+		return FElectraDecoderBitstreamProcessorH264::Create(DecoderFeatures, InitialCreationOptions);
+	}
+
 	void Suspend() override;
 	void Resume() override;
 
diff --git a/Engine/Plugins/Media/ElectraCodecs/Source/ElectraDecoders/Private/Apple/h265/VideoDecoderH265_Apple.cpp b/Engine/Plugins/Media/ElectraCodecs/Source/ElectraDecoders/Private/Apple/h265/VideoDecoderH265_Apple.cpp
--- a/Engine/Plugins/Media/ElectraCodecs/Source/ElectraDecoders/Private/Apple/h265/VideoDecoderH265_Apple.cpp
+++ b/Engine/Plugins/Media/ElectraCodecs/Source/ElectraDecoders/Private/Apple/h265/VideoDecoderH265_Apple.cpp
@@ -7,6 +7,7 @@
 #include "DecoderErrors_Apple.h"
 #include "ElectraDecodersUtils.h"
 #include "Utils/MPEG/ElectraUtilsMPEGVideo.h"
+#include "Utils/MPEG/ElectraBitstreamProcessor_H265.h"
 
 #include "IElectraDecoderFeaturesAndOptions.h"
 #include "IElectraDecoderOutputVideo.h"
@@ -176,6 +177,13 @@ public:
 	EOutputStatus HaveOutput() override;
 	TSharedPtr<IElectraDecoderOutput, ESPMode::ThreadSafe> GetOutput() override;
 
+	TSharedPtr<IElectraDecoderBitstreamProcessor, ESPMode::ThreadSafe> CreateBitstreamProcessor() override
+	{
+		TMap<FString, FVariant> DecoderFeatures;
+		GetFeatures(DecoderFeatures);
+		return FElectraDecoderBitstreamProcessorH265::Create(DecoderFeatures, InitialCreationOptions);
+	}
+
 	void Suspend() override;
 	void Resume() override;
 
diff --git a/Engine/Plugins/Media/ElectraCodecs/Source/ElectraDecoders/Private/DX/AudioDecoderAAC_DX.cpp b/Engine/Plugins/Media/ElectraCodecs/Source/ElectraDecoders/Private/DX/AudioDecoderAAC_DX.cpp
--- a/Engine/Plugins/Media/ElectraCodecs/Source/ElectraDecoders/Private/DX/AudioDecoderAAC_DX.cpp
+++ b/Engine/Plugins/Media/ElectraCodecs/Source/ElectraDecoders/Private/DX/AudioDecoderAAC_DX.cpp
@@ -6,6 +6,7 @@
 #include "DX/DecoderErrors_DX.h"
 #include "ElectraDecodersUtils.h"
 #include "Utils/MPEG/ElectraUtilsMPEGAudio.h"
+#include "Utils/ElectraBitstreamProcessorDefault.h"
 #include "DSP/FloatArrayMath.h"
 #include "IElectraDecoderFeaturesAndOptions.h"
 #include "IElectraDecoderOutputAudio.h"
@@ -116,6 +117,9 @@ public:
 	EOutputStatus HaveOutput() override;
 	TSharedPtr<IElectraDecoderOutput, ESPMode::ThreadSafe> GetOutput() override;
 
+	TSharedPtr<IElectraDecoderBitstreamProcessor, ESPMode::ThreadSafe> CreateBitstreamProcessor() override
+	{ return FElectraDecoderBitstreamProcessorDefault::Create(); }
+
 	void Suspend() override;
 	void Resume() override;
 
@@ -356,7 +360,7 @@ bool FElectraAudioDecoderAAC_DX::ResetToCleanStart()
 	CurrentOutput.Reset();
 	DecodeState = EDecodeState::Decoding;
 	bRequireDiscontinuity = true;
-	
+
 	return !LastError.IsSet();
 }
 
@@ -525,7 +529,7 @@ IElectraDecoder::EDecoderError FElectraAudioDecoderAAC_DX::SendEndOfData()
 	{
 		return IElectraDecoder::EDecoderError::EndOfData;
 	}
-	
+
 	// If there is a transform send an end-of-stream and drain message.
 	if (DecoderTransform.IsValid())
 	{
diff --git a/Engine/Plugins/Media/ElectraCodecs/Source/ElectraDecoders/Private/DX/VideoDecoderH264_DX.cpp b/Engine/Plugins/Media/ElectraCodecs/Source/ElectraDecoders/Private/DX/VideoDecoderH264_DX.cpp
--- a/Engine/Plugins/Media/ElectraCodecs/Source/ElectraDecoders/Private/DX/VideoDecoderH264_DX.cpp
+++ b/Engine/Plugins/Media/ElectraCodecs/Source/ElectraDecoders/Private/DX/VideoDecoderH264_DX.cpp
@@ -7,6 +7,7 @@
 #include "ElectraDecodersUtils.h"
 #include "Utils/MPEG/ElectraUtilsMPEGVideo.h"
 #include "Utils/MPEG/ElectraUtilsMPEGVideo_H264.h"
+#include "Utils/MPEG/ElectraBitstreamProcessor_H264.h"
 
 #include "IElectraDecoderFeaturesAndOptions.h"
 #include "IElectraDecoderOutputVideo.h"
@@ -271,6 +272,13 @@ public:
 	EOutputStatus HaveOutput() override;
 	TSharedPtr<IElectraDecoderOutput, ESPMode::ThreadSafe> GetOutput() override;
 
+	TSharedPtr<IElectraDecoderBitstreamProcessor, ESPMode::ThreadSafe> CreateBitstreamProcessor() override
+	{
+		TMap<FString, FVariant> DecoderFeatures;
+		GetFeatures(DecoderFeatures);
+		return FElectraDecoderBitstreamProcessorH264::Create(DecoderFeatures, InitialCreationOptions);
+	}
+
 	void Suspend() override;
 	void Resume() override;
 
@@ -370,6 +378,7 @@ private:
 	TRefCountPtr<IMFTransform> DecoderTransform;
 	TRefCountPtr<IMFMediaType> CurrentOutputMediaType;
 
+	TMap<FString, FVariant> InitialCreationOptions;
 	uint32 MaxOutputBuffers;
 	uint32 MaxWidth;
 	uint32 MaxHeight;
@@ -397,9 +406,9 @@ TSharedPtr<IElectraDecoder, ESPMode::ThreadSafe> IElectraVideoDecoderH264_DX::Cr
 
 
 FElectraVideoDecoderH264_DX::FElectraVideoDecoderH264_DX(const TMap<FString, FVariant>& InOptions, TSharedPtr<IElectraDecoderResourceDelegate, ESPMode::ThreadSafe> InResourceDelegate)
+	: ResourceDelegate(InResourceDelegate)
+	, InitialCreationOptions(InOptions)
 {
-	ResourceDelegate = InResourceDelegate;
-
 	MaxWidth = Align((uint32)ElectraDecodersUtil::GetVariantValueSafeU64(InOptions, TEXT("max_width"), 1920), 16u);
 	MaxHeight = Align((uint32)ElectraDecodersUtil::GetVariantValueSafeU64(InOptions, TEXT("max_height"), 1080), 16u);
 
diff --git a/Engine/Plugins/Media/ElectraCodecs/Source/ElectraDecoders/Private/DX/VideoDecoderH265_DX.cpp b/Engine/Plugins/Media/ElectraCodecs/Source/ElectraDecoders/Private/DX/VideoDecoderH265_DX.cpp
--- a/Engine/Plugins/Media/ElectraCodecs/Source/ElectraDecoders/Private/DX/VideoDecoderH265_DX.cpp
+++ b/Engine/Plugins/Media/ElectraCodecs/Source/ElectraDecoders/Private/DX/VideoDecoderH265_DX.cpp
@@ -6,6 +6,7 @@
 #include "DX/DecoderErrors_DX.h"
 #include "ElectraDecodersUtils.h"
 #include "Utils/MPEG/ElectraUtilsMPEGVideo.h"
+#include "Utils/MPEG/ElectraBitstreamProcessor_H265.h"
 
 #include "IElectraDecoderFeaturesAndOptions.h"
 #include "IElectraDecoderOutputVideo.h"
@@ -270,6 +271,13 @@ public:
 	EOutputStatus HaveOutput() override;
 	TSharedPtr<IElectraDecoderOutput, ESPMode::ThreadSafe> GetOutput() override;
 
+	TSharedPtr<IElectraDecoderBitstreamProcessor, ESPMode::ThreadSafe> CreateBitstreamProcessor() override
+	{
+		TMap<FString, FVariant> DecoderFeatures;
+		GetFeatures(DecoderFeatures);
+		return FElectraDecoderBitstreamProcessorH265::Create(DecoderFeatures, InitialCreationOptions);
+	}
+
 	void Suspend() override;
 	void Resume() override;
 
@@ -367,6 +375,7 @@ private:
 	TRefCountPtr<IMFTransform> DecoderTransform;
 	TRefCountPtr<IMFMediaType> CurrentOutputMediaType;
 
+	TMap<FString, FVariant> InitialCreationOptions;
 	uint32 MaxOutputBuffers;
 	uint32 MaxWidth;
 	uint32 MaxHeight;
@@ -394,9 +403,9 @@ TSharedPtr<IElectraDecoder, ESPMode::ThreadSafe> IElectraVideoDecoderH265_DX::Cr
 
 
 FElectraVideoDecoderH265_DX::FElectraVideoDecoderH265_DX(const TMap<FString, FVariant>& InOptions, TSharedPtr<IElectraDecoderResourceDelegate, ESPMode::ThreadSafe> InResourceDelegate)
+	: ResourceDelegate(InResourceDelegate)
+	, InitialCreationOptions(InOptions)
 {
-	ResourceDelegate = InResourceDelegate;
-
 	MaxWidth = Align((uint32)ElectraDecodersUtil::GetVariantValueSafeU64(InOptions, TEXT("max_width"), 3840), 16u);
 	MaxHeight = Align((uint32)ElectraDecodersUtil::GetVariantValueSafeU64(InOptions, TEXT("max_height"), 2160), 16u);
 
diff --git a/Engine/Plugins/Media/ElectraCodecs/Source/ElectraDecoders/Private/ElectraMediaMP3Decoder.cpp b/Engine/Plugins/Media/ElectraCodecs/Source/ElectraDecoders/Private/ElectraMediaMP3Decoder.cpp
--- a/Engine/Plugins/Media/ElectraCodecs/Source/ElectraDecoders/Private/ElectraMediaMP3Decoder.cpp
+++ b/Engine/Plugins/Media/ElectraCodecs/Source/ElectraDecoders/Private/ElectraMediaMP3Decoder.cpp
@@ -15,6 +15,7 @@
 #include "ElectraDecodersUtils.h"
 #include "Utils/ElectraBitstreamReader.h"
 #include "Utils/MPEG/ElectraUtilsMPEGAudio.h"
+#include "Utils/ElectraBitstreamProcessorDefault.h"
 
 #ifndef WITH_DR_LIBS_MP3
 #define WITH_DR_LIBS_MP3 0
@@ -156,6 +157,9 @@ public:
 	EOutputStatus HaveOutput() override;
 	TSharedPtr<IElectraDecoderOutput, ESPMode::ThreadSafe> GetOutput() override;
 
+	TSharedPtr<IElectraDecoderBitstreamProcessor, ESPMode::ThreadSafe> CreateBitstreamProcessor() override
+	{ return FElectraDecoderBitstreamProcessorDefault::Create(); }
+
 	void Suspend() override
 	{ }
 	void Resume() override
diff --git a/Engine/Plugins/Media/ElectraCodecs/Source/ElectraDecoders/Private/Linux/aac/AudioDecoderAAC_Linux.cpp b/Engine/Plugins/Media/ElectraCodecs/Source/ElectraDecoders/Private/Linux/aac/AudioDecoderAAC_Linux.cpp
--- a/Engine/Plugins/Media/ElectraCodecs/Source/ElectraDecoders/Private/Linux/aac/AudioDecoderAAC_Linux.cpp
+++ b/Engine/Plugins/Media/ElectraCodecs/Source/ElectraDecoders/Private/Linux/aac/AudioDecoderAAC_Linux.cpp
@@ -7,6 +7,7 @@
 #include "DecoderErrors_Linux.h"
 #include "ElectraDecodersUtils.h"
 #include "Utils/MPEG/ElectraUtilsMPEGAudio.h"
+#include "Utils/ElectraBitstreamProcessorDefault.h"
 
 #include "IElectraDecoderFeaturesAndOptions.h"
 #include "IElectraDecoderOutputAudio.h"
@@ -113,6 +114,9 @@ public:
 	EOutputStatus HaveOutput() override;
 	TSharedPtr<IElectraDecoderOutput, ESPMode::ThreadSafe> GetOutput() override;
 
+	TSharedPtr<IElectraDecoderBitstreamProcessor, ESPMode::ThreadSafe> CreateBitstreamProcessor() override
+	{ return FElectraDecoderBitstreamProcessorDefault::Create(); }
+
 	void Suspend() override;
 	void Resume() override;
 
diff --git a/Engine/Plugins/Media/ElectraCodecs/Source/ElectraDecoders/Private/Linux/h264/VideoDecoderH264_Linux.cpp b/Engine/Plugins/Media/ElectraCodecs/Source/ElectraDecoders/Private/Linux/h264/VideoDecoderH264_Linux.cpp
--- a/Engine/Plugins/Media/ElectraCodecs/Source/ElectraDecoders/Private/Linux/h264/VideoDecoderH264_Linux.cpp
+++ b/Engine/Plugins/Media/ElectraCodecs/Source/ElectraDecoders/Private/Linux/h264/VideoDecoderH264_Linux.cpp
@@ -4,6 +4,7 @@
 #include "DecoderErrors_Linux.h"
 #include "ElectraDecodersUtils.h"
 #include "Utils/MPEG/ElectraUtilsMPEGVideo.h"
+#include "Utils/MPEG/ElectraBitstreamProcessor_H264.h"
 
 #include "IElectraDecoderFeaturesAndOptions.h"
 #include "IElectraDecoderOutputVideo.h"
@@ -125,6 +126,13 @@ public:
 	EOutputStatus HaveOutput() override;
 	TSharedPtr<IElectraDecoderOutput, ESPMode::ThreadSafe> GetOutput() override;
 
+	TSharedPtr<IElectraDecoderBitstreamProcessor, ESPMode::ThreadSafe> CreateBitstreamProcessor() override
+	{
+		TMap<FString, FVariant> DecoderFeatures;
+		GetFeatures(DecoderFeatures);
+		return FElectraDecoderBitstreamProcessorH264::Create(DecoderFeatures, InitialCreationOptions);
+	}
+
 	void Suspend() override;
 	void Resume() override;
 
diff --git a/Engine/Plugins/Media/ElectraCodecs/Source/ElectraDecoders/Private/Linux/h265/VideoDecoderH265_Linux.cpp b/Engine/Plugins/Media/ElectraCodecs/Source/ElectraDecoders/Private/Linux/h265/VideoDecoderH265_Linux.cpp
--- a/Engine/Plugins/Media/ElectraCodecs/Source/ElectraDecoders/Private/Linux/h265/VideoDecoderH265_Linux.cpp
+++ b/Engine/Plugins/Media/ElectraCodecs/Source/ElectraDecoders/Private/Linux/h265/VideoDecoderH265_Linux.cpp
@@ -4,6 +4,7 @@
 #include "DecoderErrors_Linux.h"
 #include "ElectraDecodersUtils.h"
 #include "Utils/MPEG/ElectraUtilsMPEGVideo.h"
+#include "Utils/MPEG/ElectraBitstreamProcessor_H265.h"
 
 #include "IElectraDecoderFeaturesAndOptions.h"
 #include "IElectraDecoderOutputVideo.h"
@@ -125,6 +126,13 @@ public:
 	EOutputStatus HaveOutput() override;
 	TSharedPtr<IElectraDecoderOutput, ESPMode::ThreadSafe> GetOutput() override;
 
+	TSharedPtr<IElectraDecoderBitstreamProcessor, ESPMode::ThreadSafe> CreateBitstreamProcessor() override
+	{
+		TMap<FString, FVariant> DecoderFeatures;
+		GetFeatures(DecoderFeatures);
+		return FElectraDecoderBitstreamProcessorH265::Create(DecoderFeatures, InitialCreationOptions);
+	}
+
 	void Suspend() override;
 	void Resume() override;
 
diff --git a/Engine/Plugins/Media/ElectraCodecs/Source/ElectraDecoders/Private/Linux/vpx/VideoDecoderVPx_Linux.cpp b/Engine/Plugins/Media/ElectraCodecs/Source/ElectraDecoders/Private/Linux/vpx/VideoDecoderVPx_Linux.cpp
--- a/Engine/Plugins/Media/ElectraCodecs/Source/ElectraDecoders/Private/Linux/vpx/VideoDecoderVPx_Linux.cpp
+++ b/Engine/Plugins/Media/ElectraCodecs/Source/ElectraDecoders/Private/Linux/vpx/VideoDecoderVPx_Linux.cpp
@@ -4,6 +4,7 @@
 #include "DecoderErrors_Linux.h"
 #include "ElectraDecodersUtils.h"
 #include "Utils/Google/ElectraUtilsVPxVideo.h"
+#include "Utils/Google/ElectraBitstreamProcessorVPx.h"
 
 #include "IElectraDecoderFeaturesAndOptions.h"
 #include "IElectraDecoderOutputVideo.h"
@@ -126,6 +127,13 @@ public:
 	EOutputStatus HaveOutput() override;
 	TSharedPtr<IElectraDecoderOutput, ESPMode::ThreadSafe> GetOutput() override;
 
+	TSharedPtr<IElectraDecoderBitstreamProcessor, ESPMode::ThreadSafe> CreateBitstreamProcessor() override
+	{
+		TMap<FString, FVariant> DecoderFeatures;
+		GetFeatures(DecoderFeatures);
+		return FElectraDecoderBitstreamProcessorVPx::Create(DecoderFeatures, InitialCreationOptions);
+	}
+
 	void Suspend() override;
 	void Resume() override;
 
diff --git a/Engine/Plugins/Media/ElectraCodecs/Source/ElectraDecoders/Private/Utils/ElectraBitstreamProcessorDefault.cpp b/Engine/Plugins/Media/ElectraCodecs/Source/ElectraDecoders/Private/Utils/ElectraBitstreamProcessorDefault.cpp
new file mode 100644
--- /dev/null
+++ b/Engine/Plugins/Media/ElectraCodecs/Source/ElectraDecoders/Private/Utils/ElectraBitstreamProcessorDefault.cpp
@@ -0,0 +1,50 @@
+// Copyright Epic Games, Inc. All Rights Reserved.
+
+#include "Utils/ElectraBitstreamProcessorDefault.h"
+
+
+/*
+class FElectraDecoderBitstreamProcessorDefault::FImpl
+{
+public:
+};
+*/
+
+FElectraDecoderBitstreamProcessorDefault::~FElectraDecoderBitstreamProcessorDefault()
+{
+}
+
+FElectraDecoderBitstreamProcessorDefault::FElectraDecoderBitstreamProcessorDefault()
+{
+//	Impl = MakePimpl<FImpl>();
+}
+
+bool FElectraDecoderBitstreamProcessorDefault::WillModifyBitstreamInPlace()
+{
+	return false;
+}
+
+void FElectraDecoderBitstreamProcessorDefault::Clear()
+{
+}
+
+IElectraDecoderBitstreamProcessor::EProcessResult FElectraDecoderBitstreamProcessorDefault::GetCSDFromConfigurationRecord(TArray<uint8>& OutCSD, const TMap<FString, FVariant>& InParamsWithDCRorCSD)
+{
+	OutCSD.Empty();
+	return EProcessResult::Ok;
+}
+
+IElectraDecoderBitstreamProcessor::EProcessResult FElectraDecoderBitstreamProcessorDefault::ProcessInputForDecoding(TSharedPtr<IElectraDecoderBitstreamInfo, ESPMode::ThreadSafe>& OutBSI, FElectraDecoderInputAccessUnit& InOutAccessUnit, const TMap<FString, FVariant>& InAccessUnitSidebandData)
+{
+	InOutAccessUnit.Flags |= EElectraDecoderFlags::InputIsProcessed;
+	return EProcessResult::Ok;
+}
+
+void FElectraDecoderBitstreamProcessorDefault::SetPropertiesOnOutput(TMap<FString, FVariant>& InOutProperties, TSharedPtr<IElectraDecoderBitstreamInfo, ESPMode::ThreadSafe> InBSI)
+{
+}
+
+FString FElectraDecoderBitstreamProcessorDefault::GetLastError()
+{
+	return FString();
+}
diff --git a/Engine/Plugins/Media/ElectraCodecs/Source/ElectraDecoders/Private/Utils/ElectraBitstreamProcessorGenericVideo.cpp b/Engine/Plugins/Media/ElectraCodecs/Source/ElectraDecoders/Private/Utils/ElectraBitstreamProcessorGenericVideo.cpp
new file mode 100644
--- /dev/null
+++ b/Engine/Plugins/Media/ElectraCodecs/Source/ElectraDecoders/Private/Utils/ElectraBitstreamProcessorGenericVideo.cpp
@@ -0,0 +1,189 @@
+// Copyright Epic Games, Inc. All Rights Reserved.
+
+#include "Utils/ElectraBitstreamProcessorGenericVideo.h"
+#include "Utils/MPEG/ElectraUtilsMPEGVideo.h"
+#include "Utils/MPEG/ElectraUtilsMP4.h"
+#include "IElectraDecoderFeaturesAndOptions.h"
+#include "ElectraDecodersUtils.h"
+
+
+
+class FElectraDecoderBitstreamProcessorGenericVideo::FImpl
+{
+public:
+	void Clear()
+	{
+		bSentColorimetry = false;
+		bSentMDCV = false;
+		bSentCLLI = false;
+		// Do not clear out any error message!
+	}
+
+	FString GetLastError()
+	{
+		return LastErrorMessage;
+	}
+
+	void SetPropertiesOnOutput(TMap<FString, FVariant>& InOutProperties)
+	{
+		if (CurrentColorimetry.IsValid() && !bSentColorimetry)
+		{
+			bSentColorimetry = true;
+			TConstArrayView<uint8> Colorimetry = MakeConstArrayView(reinterpret_cast<const uint8*>(CurrentColorimetry.Get()), sizeof(ElectraDecodersUtil::MPEG::FCommonColorimetry));
+			InOutProperties.Emplace(IElectraDecoderBitstreamProcessorInfo::CommonColorimetry, TArray<uint8>(Colorimetry));
+		}
+
+		if (CurrentMDCV.IsValid() && !bSentMDCV)
+		{
+			bSentMDCV = true;
+			TConstArrayView<uint8> Params = MakeConstArrayView(reinterpret_cast<const uint8*>(CurrentMDCV.Get()), sizeof(ElectraDecodersUtil::MPEG::FSEImastering_display_colour_volume));
+			InOutProperties.Emplace(IElectraDecoderBitstreamProcessorInfo::SeiMasteringDisplayColorVolume, TArray<uint8>(Params));
+		}
+
+		if (CurrentCLLI.IsValid() && !bSentCLLI)
+		{
+			bSentCLLI = true;
+			TConstArrayView<uint8> Params = MakeConstArrayView(reinterpret_cast<const uint8*>(CurrentCLLI.Get()), sizeof(ElectraDecodersUtil::MPEG::FSEIcontent_light_level_info));
+			InOutProperties.Emplace(IElectraDecoderBitstreamProcessorInfo::SeiContentLightLeveInfo, TArray<uint8>(Params));
+		}
+	}
+
+	void SetColorimetryFromCOLRBox(const TArray<uint8>& InCOLRBox)
+	{
+		COLRBox = InCOLRBox;
+		UpdateCOLR();
+	}
+	void SetMasteringDisplayColorVolumeFromMDCVBox(const TArray<uint8>& InMDCVBox)
+	{
+		MDCVBox = InMDCVBox;
+		UpdateMDCV();
+	}
+	void SetContentLightLevelInfoFromCLLIBox(const TArray<uint8>& InCLLIBox)
+	{
+		CLLIBox = InCLLIBox;
+		UpdateCLLI();
+	}
+	void SetContentLightLevelFromCOLLBox(const TArray<uint8>& InCOLLBox)
+	{
+		if (InCOLLBox.Num() > 4)
+		{
+			uint8 BoxVersion = InCOLLBox[0];
+			if (BoxVersion == 0)
+			{
+				// 'clli' box is the same as a version 0 'COLL 'coll' box.
+				CLLIBox = TArray<uint8>(InCOLLBox.GetData() + 4, InCOLLBox.Num() - 4);
+			}
+		}
+		UpdateCLLI();
+	}
+private:
+	void UpdateCOLR()
+	{
+		if (COLRBox.IsEmpty())
+		{
+			return;
+		}
+		TSharedPtr<ElectraDecodersUtil::MPEG::FCommonColorimetry, ESPMode::ThreadSafe> NewColorimetry = MakeShared<ElectraDecodersUtil::MPEG::FCommonColorimetry, ESPMode::ThreadSafe>();
+		if (ElectraDecodersUtil::MPEG::ParseFromCOLRBox(*NewColorimetry, COLRBox))
+		{
+			CurrentColorimetry = MoveTemp(NewColorimetry);
+		}
+		else
+		{
+			LastErrorMessage = TEXT("Failed to parse `colr` box data");
+		}
+	}
+
+	void UpdateMDCV()
+	{
+		if (MDCVBox.IsEmpty())
+		{
+			return;
+		}
+		TSharedPtr<ElectraDecodersUtil::MPEG::FSEImastering_display_colour_volume, ESPMode::ThreadSafe> NewMDCV = MakeShared<ElectraDecodersUtil::MPEG::FSEImastering_display_colour_volume, ESPMode::ThreadSafe>();
+		if (ElectraDecodersUtil::MPEG::ParseFromMDCVBox(*NewMDCV, MDCVBox))
+		{
+			CurrentMDCV = MoveTemp(NewMDCV);
+		}
+		else
+		{
+			LastErrorMessage = TEXT("Failed to parse `mdcv` box data");
+		}
+	}
+
+	void UpdateCLLI()
+	{
+		if (CLLIBox.IsEmpty())
+		{
+			return;
+		}
+		TSharedPtr<ElectraDecodersUtil::MPEG::FSEIcontent_light_level_info, ESPMode::ThreadSafe> NewCLLI = MakeShared<ElectraDecodersUtil::MPEG::FSEIcontent_light_level_info, ESPMode::ThreadSafe>();
+		if (ElectraDecodersUtil::MPEG::ParseFromCLLIBox(*NewCLLI, CLLIBox))
+		{
+			CurrentCLLI = MoveTemp(NewCLLI);
+		}
+		else
+		{
+			LastErrorMessage = TEXT("Failed to parse `coll`/`clli` box data");
+		}
+	}
+
+	TSharedPtr<ElectraDecodersUtil::MPEG::FCommonColorimetry, ESPMode::ThreadSafe> CurrentColorimetry;
+	TSharedPtr<ElectraDecodersUtil::MPEG::FSEImastering_display_colour_volume, ESPMode::ThreadSafe> CurrentMDCV;
+	TSharedPtr<ElectraDecodersUtil::MPEG::FSEIcontent_light_level_info, ESPMode::ThreadSafe> CurrentCLLI;
+	TArray<uint8> COLRBox;
+	TArray<uint8> MDCVBox;
+	TArray<uint8> CLLIBox;
+	FString LastErrorMessage;
+	bool bSentColorimetry = false;
+	bool bSentMDCV = false;
+	bool bSentCLLI = false;
+};
+
+
+FElectraDecoderBitstreamProcessorGenericVideo::~FElectraDecoderBitstreamProcessorGenericVideo()
+{
+}
+
+FElectraDecoderBitstreamProcessorGenericVideo::FElectraDecoderBitstreamProcessorGenericVideo(const TMap<FString, FVariant>& InDecoderParams, const TMap<FString, FVariant>& InFormatParams)
+{
+	Impl = MakePimpl<FImpl>();
+	Impl->SetColorimetryFromCOLRBox(ElectraDecodersUtil::GetVariantValueUInt8Array(InFormatParams, TEXT("$colr_box")));
+	Impl->SetMasteringDisplayColorVolumeFromMDCVBox(ElectraDecodersUtil::GetVariantValueUInt8Array(InFormatParams, TEXT("$mdcv_box")));
+	Impl->SetContentLightLevelFromCOLLBox(ElectraDecodersUtil::GetVariantValueUInt8Array(InFormatParams, TEXT("$coll_box")));
+	Impl->SetContentLightLevelInfoFromCLLIBox(ElectraDecodersUtil::GetVariantValueUInt8Array(InFormatParams, TEXT("$clli_box")));
+}
+
+bool FElectraDecoderBitstreamProcessorGenericVideo::WillModifyBitstreamInPlace()
+{
+	return false;
+}
+
+void FElectraDecoderBitstreamProcessorGenericVideo::Clear()
+{
+	// This only clears the flags that we did not send the color parameters yet.
+	// The actual parameters are left unchanged as they are one-time init on construction only.
+	Impl->Clear();
+}
+
+IElectraDecoderBitstreamProcessor::EProcessResult FElectraDecoderBitstreamProcessorGenericVideo::GetCSDFromConfigurationRecord(TArray<uint8>& OutCSD, const TMap<FString, FVariant>& InParamsWithDCRorCSD)
+{
+	OutCSD.Empty();
+	return EProcessResult::Ok;
+}
+
+IElectraDecoderBitstreamProcessor::EProcessResult FElectraDecoderBitstreamProcessorGenericVideo::ProcessInputForDecoding(TSharedPtr<IElectraDecoderBitstreamInfo, ESPMode::ThreadSafe>& OutBSI, FElectraDecoderInputAccessUnit& InOutAccessUnit, const TMap<FString, FVariant>& InAccessUnitSidebandData)
+{
+	InOutAccessUnit.Flags |= EElectraDecoderFlags::InputIsProcessed;
+	return EProcessResult::Ok;
+}
+
+void FElectraDecoderBitstreamProcessorGenericVideo::SetPropertiesOnOutput(TMap<FString, FVariant>& InOutProperties, TSharedPtr<IElectraDecoderBitstreamInfo, ESPMode::ThreadSafe> InBSI)
+{
+	Impl->SetPropertiesOnOutput(InOutProperties);
+}
+
+FString FElectraDecoderBitstreamProcessorGenericVideo::GetLastError()
+{
+	return Impl->GetLastError();
+}
diff --git a/Engine/Plugins/Media/ElectraCodecs/Source/ElectraDecoders/Private/Utils/Google/ElectraBitstreamProcessorVPx.cpp b/Engine/Plugins/Media/ElectraCodecs/Source/ElectraDecoders/Private/Utils/Google/ElectraBitstreamProcessorVPx.cpp
new file mode 100644
--- /dev/null
+++ b/Engine/Plugins/Media/ElectraCodecs/Source/ElectraDecoders/Private/Utils/Google/ElectraBitstreamProcessorVPx.cpp
@@ -0,0 +1,262 @@
+// Copyright Epic Games, Inc. All Rights Reserved.
+
+#include "Utils/Google/ElectraBitstreamProcessorVPx.h"
+#include "Utils/Google/ElectraUtilsVPxVideo.h"
+#include "Utils/MPEG/ElectraUtilsMPEGVideo.h"	// because information may be carried in mp4 specific boxes
+#include "Utils/MPEG/ElectraUtilsMP4.h"
+#include "IElectraDecoderFeaturesAndOptions.h"
+#include "ElectraDecodersUtils.h"
+
+
+
+class FElectraDecoderBitstreamProcessorVPx::FImpl
+{
+public:
+	enum class EVersion
+	{
+		VP8,
+		VP9
+	};
+
+	FImpl(EVersion InVersion)
+		: Version(InVersion)
+	{ }
+
+	void Clear()
+	{
+		bSentColorimetry = false;
+		bSentMDCV = false;
+		bSentCLLI = false;
+		// Do not clear out any error message!
+	}
+
+	FString GetLastError()
+	{
+		return LastErrorMessage;
+	}
+
+	bool ProcessInputForDecoding(FElectraDecoderInputAccessUnit& InOutAccessUnit)
+	{
+		InOutAccessUnit.Flags &= ~EElectraDecoderFlags::IsDiscardable;
+		if (Version == EVersion::VP8)
+		{
+			ElectraDecodersUtil::VPxVideo::FVP8UncompressedHeader Header;
+			if (ElectraDecodersUtil::VPxVideo::ParseVP8UncompressedHeader(Header, InOutAccessUnit.Data, InOutAccessUnit.DataSize))
+			{
+				if (Header.IsKeyframe())
+				{
+					InOutAccessUnit.Flags |= EElectraDecoderFlags::IsSyncSample;
+				}
+				else
+				{
+					InOutAccessUnit.Flags &= ~EElectraDecoderFlags::IsSyncSample;
+				}
+				return true;
+			}
+			else
+			{
+				LastErrorMessage = FString::Printf(TEXT("Failed to parse VP8 header"));
+			}
+		}
+		else if (Version == EVersion::VP9)
+		{
+			ElectraDecodersUtil::VPxVideo::FVP9UncompressedHeader Header;
+			if (ElectraDecodersUtil::VPxVideo::ParseVP9UncompressedHeader(Header, InOutAccessUnit.Data, InOutAccessUnit.DataSize))
+			{
+				if (Header.IsKeyframe())
+				{
+					InOutAccessUnit.Flags |= EElectraDecoderFlags::IsSyncSample;
+				}
+				else
+				{
+					InOutAccessUnit.Flags &= ~EElectraDecoderFlags::IsSyncSample;
+				}
+				return true;
+			}
+			else
+			{
+				LastErrorMessage = FString::Printf(TEXT("Failed to parse VP9 header"));
+			}
+		}
+		return false;
+	}
+
+	void SetPropertiesOnOutput(TMap<FString, FVariant>& InOutProperties)
+	{
+		if (CurrentColorimetry.IsValid() && !bSentColorimetry)
+		{
+			bSentColorimetry = true;
+			TConstArrayView<uint8> Colorimetry = MakeConstArrayView(reinterpret_cast<const uint8*>(CurrentColorimetry.Get()), sizeof(ElectraDecodersUtil::MPEG::FCommonColorimetry));
+			InOutProperties.Emplace(IElectraDecoderBitstreamProcessorInfo::CommonColorimetry, TArray<uint8>(Colorimetry));
+		}
+
+		if (CurrentMDCV.IsValid() && !bSentMDCV)
+		{
+			bSentMDCV = true;
+			TConstArrayView<uint8> Params = MakeConstArrayView(reinterpret_cast<const uint8*>(CurrentMDCV.Get()), sizeof(ElectraDecodersUtil::MPEG::FSEImastering_display_colour_volume));
+			InOutProperties.Emplace(IElectraDecoderBitstreamProcessorInfo::SeiMasteringDisplayColorVolume, TArray<uint8>(Params));
+		}
+
+		if (CurrentCLLI.IsValid() && !bSentCLLI)
+		{
+			bSentCLLI = true;
+			TConstArrayView<uint8> Params = MakeConstArrayView(reinterpret_cast<const uint8*>(CurrentCLLI.Get()), sizeof(ElectraDecodersUtil::MPEG::FSEIcontent_light_level_info));
+			InOutProperties.Emplace(IElectraDecoderBitstreamProcessorInfo::SeiContentLightLeveInfo, TArray<uint8>(Params));
+		}
+	}
+
+	void SetColorimetryFromCOLRBox(const TArray<uint8>& InCOLRBox)
+	{
+		COLRBox = InCOLRBox;
+		UpdateCOLR();
+	}
+	void SetMasteringDisplayColorVolumeFromMDCVBox(const TArray<uint8>& InMDCVBox)
+	{
+		MDCVBox = InMDCVBox;
+		UpdateMDCV();
+	}
+	void SetContentLightLevelInfoFromCLLIBox(const TArray<uint8>& InCLLIBox)
+	{
+		CLLIBox = InCLLIBox;
+		UpdateCLLI();
+	}
+	void SetContentLightLevelFromCOLLBox(const TArray<uint8>& InCOLLBox)
+	{
+		if (InCOLLBox.Num() > 4)
+		{
+			uint8 BoxVersion = InCOLLBox[0];
+			if (BoxVersion == 0)
+			{
+				// 'clli' box is the same as a version 0 'COLL 'coll' box.
+				CLLIBox = TArray<uint8>(InCOLLBox.GetData() + 4, InCOLLBox.Num() - 4);
+			}
+		}
+		UpdateCLLI();
+	}
+private:
+	FImpl() = delete;
+
+	void UpdateCOLR()
+	{
+		if (COLRBox.IsEmpty())
+		{
+			return;
+		}
+		TSharedPtr<ElectraDecodersUtil::MPEG::FCommonColorimetry, ESPMode::ThreadSafe> NewColorimetry = MakeShared<ElectraDecodersUtil::MPEG::FCommonColorimetry, ESPMode::ThreadSafe>();
+		if (ElectraDecodersUtil::MPEG::ParseFromCOLRBox(*NewColorimetry, COLRBox))
+		{
+			CurrentColorimetry = MoveTemp(NewColorimetry);
+		}
+		else
+		{
+			LastErrorMessage = TEXT("Failed to parse `colr` box data");
+		}
+	}
+
+	void UpdateMDCV()
+	{
+		if (MDCVBox.IsEmpty())
+		{
+			return;
+		}
+		TSharedPtr<ElectraDecodersUtil::MPEG::FSEImastering_display_colour_volume, ESPMode::ThreadSafe> NewMDCV = MakeShared<ElectraDecodersUtil::MPEG::FSEImastering_display_colour_volume, ESPMode::ThreadSafe>();
+		if (ElectraDecodersUtil::MPEG::ParseFromMDCVBox(*NewMDCV, MDCVBox))
+		{
+			CurrentMDCV = MoveTemp(NewMDCV);
+		}
+		else
+		{
+			LastErrorMessage = TEXT("Failed to parse `mdcv` box data");
+		}
+	}
+
+	void UpdateCLLI()
+	{
+		if (CLLIBox.IsEmpty())
+		{
+			return;
+		}
+		TSharedPtr<ElectraDecodersUtil::MPEG::FSEIcontent_light_level_info, ESPMode::ThreadSafe> NewCLLI = MakeShared<ElectraDecodersUtil::MPEG::FSEIcontent_light_level_info, ESPMode::ThreadSafe>();
+		if (ElectraDecodersUtil::MPEG::ParseFromCLLIBox(*NewCLLI, CLLIBox))
+		{
+			CurrentCLLI = MoveTemp(NewCLLI);
+		}
+		else
+		{
+			LastErrorMessage = TEXT("Failed to parse `coll`/`clli` box data");
+		}
+	}
+
+	EVersion Version = EVersion::VP8;
+	TSharedPtr<ElectraDecodersUtil::MPEG::FCommonColorimetry, ESPMode::ThreadSafe> CurrentColorimetry;
+	TSharedPtr<ElectraDecodersUtil::MPEG::FSEImastering_display_colour_volume, ESPMode::ThreadSafe> CurrentMDCV;
+	TSharedPtr<ElectraDecodersUtil::MPEG::FSEIcontent_light_level_info, ESPMode::ThreadSafe> CurrentCLLI;
+	TArray<uint8> COLRBox;
+	TArray<uint8> MDCVBox;
+	TArray<uint8> CLLIBox;
+	FString LastErrorMessage;
+	bool bSentColorimetry = false;
+	bool bSentMDCV = false;
+	bool bSentCLLI = false;
+};
+
+
+FElectraDecoderBitstreamProcessorVPx::~FElectraDecoderBitstreamProcessorVPx()
+{
+}
+
+FElectraDecoderBitstreamProcessorVPx::FElectraDecoderBitstreamProcessorVPx(const TMap<FString, FVariant>& InDecoderParams, const TMap<FString, FVariant>& InFormatParams)
+{
+	uint32 Codec4CC = (uint32)ElectraDecodersUtil::GetVariantValueSafeU64(InFormatParams, TEXT("codec_4cc"), 0);
+	check(Codec4CC == ElectraDecodersUtil::Make4CC('v','p','0','8') || Codec4CC == ElectraDecodersUtil::Make4CC('v','p','0','9'));
+
+	Impl = MakePimpl<FImpl>(Codec4CC == ElectraDecodersUtil::Make4CC('v','p','0','8') ? FImpl::EVersion::VP8 : FImpl::EVersion::VP9);
+	Impl->SetColorimetryFromCOLRBox(ElectraDecodersUtil::GetVariantValueUInt8Array(InFormatParams, TEXT("$colr_box")));
+	Impl->SetMasteringDisplayColorVolumeFromMDCVBox(ElectraDecodersUtil::GetVariantValueUInt8Array(InFormatParams, TEXT("$mdcv_box")));
+	Impl->SetContentLightLevelFromCOLLBox(ElectraDecodersUtil::GetVariantValueUInt8Array(InFormatParams, TEXT("$coll_box")));
+	Impl->SetContentLightLevelInfoFromCLLIBox(ElectraDecodersUtil::GetVariantValueUInt8Array(InFormatParams, TEXT("$clli_box")));
+}
+
+bool FElectraDecoderBitstreamProcessorVPx::WillModifyBitstreamInPlace()
+{
+	return false;
+}
+
+void FElectraDecoderBitstreamProcessorVPx::Clear()
+{
+	// This only clears the flags that we did not send the color parameters yet.
+	// The actual parameters are left unchanged as they are one-time init on construction only.
+	Impl->Clear();
+}
+
+IElectraDecoderBitstreamProcessor::EProcessResult FElectraDecoderBitstreamProcessorVPx::GetCSDFromConfigurationRecord(TArray<uint8>& OutCSD, const TMap<FString, FVariant>& InParamsWithDCRorCSD)
+{
+	OutCSD.Empty();
+	return EProcessResult::Ok;
+}
+
+IElectraDecoderBitstreamProcessor::EProcessResult FElectraDecoderBitstreamProcessorVPx::ProcessInputForDecoding(TSharedPtr<IElectraDecoderBitstreamInfo, ESPMode::ThreadSafe>& OutBSI, FElectraDecoderInputAccessUnit& InOutAccessUnit, const TMap<FString, FVariant>& InAccessUnitSidebandData)
+{
+	if ((InOutAccessUnit.Flags & EElectraDecoderFlags::InputIsProcessed) == EElectraDecoderFlags::InputIsProcessed)
+	{
+		return EProcessResult::Ok;
+	}
+
+	InOutAccessUnit.Flags |= EElectraDecoderFlags::InputIsProcessed;
+	if (!Impl->ProcessInputForDecoding(InOutAccessUnit))
+	{
+		return EProcessResult::Error;
+	}
+
+	return EProcessResult::Ok;
+}
+
+void FElectraDecoderBitstreamProcessorVPx::SetPropertiesOnOutput(TMap<FString, FVariant>& InOutProperties, TSharedPtr<IElectraDecoderBitstreamInfo, ESPMode::ThreadSafe> InBSI)
+{
+	Impl->SetPropertiesOnOutput(InOutProperties);
+}
+
+FString FElectraDecoderBitstreamProcessorVPx::GetLastError()
+{
+	return Impl->GetLastError();
+}
diff --git a/Engine/Plugins/Media/ElectraCodecs/Source/ElectraDecoders/Private/Utils/MPEG/ElectraBitstreamProcessor_H264.cpp b/Engine/Plugins/Media/ElectraCodecs/Source/ElectraDecoders/Private/Utils/MPEG/ElectraBitstreamProcessor_H264.cpp
new file mode 100644
--- /dev/null
+++ b/Engine/Plugins/Media/ElectraCodecs/Source/ElectraDecoders/Private/Utils/MPEG/ElectraBitstreamProcessor_H264.cpp
@@ -0,0 +1,415 @@
+// Copyright Epic Games, Inc. All Rights Reserved.
+
+#include "Utils/MPEG/ElectraBitstreamProcessor_H264.h"
+#include "Utils/MPEG/ElectraUtilsMPEGVideo_H264.h"
+#include "Utils/MPEG/ElectraUtilsMPEGVideo.h"
+#include "Utils/MPEG/ElectraUtilsMP4.h"
+#include "IElectraDecoderFeaturesAndOptions.h"
+#include "ElectraDecodersUtils.h"
+
+
+class FElectraDecoderBitstreamProcessorH264::FImpl
+{
+public:
+	class FBitstreamInfo : public IElectraDecoderBitstreamInfo
+	{
+	public:
+		virtual ~FBitstreamInfo() = default;
+		TArray<ElectraDecodersUtil::MPEG::FSEIMessage> SEIMessages;
+		TMap<uint32, ElectraDecodersUtil::MPEG::H264::FSequenceParameterSet> SPSs;
+	};
+
+	enum class ECSDResult
+	{
+		NoCSD,
+		Unchanged,
+		Changed,
+		Error
+	};
+
+	ECSDResult ExtractSPS(const TMap<FString, FVariant>& InFromMap)
+	{
+		// First try the `$avcC_box`
+		TArray<uint8> ConfigData = ElectraDecodersUtil::GetVariantValueUInt8Array(InFromMap, AvccBoxName);
+		// If empty try the decoder confiuration record, which is the same as the ConfigData box.
+		if (ConfigData.IsEmpty())
+		{
+			ConfigData = ElectraDecodersUtil::GetVariantValueUInt8Array(InFromMap, DCRName);
+		}
+		if (ConfigData.Num())
+		{
+			// If the data is the same then we don't need to reparse it.
+			if (ConfigData == CurrentDecoderConfiguration)
+			{
+				return ECSDResult::Unchanged;
+			}
+			SPSs.Empty();
+			ElectraDecodersUtil::MPEG::H264::FAVCDecoderConfigurationRecord dcr;
+			if (!dcr.Parse(ConfigData) || dcr.GetSequenceParameterSets().IsEmpty())
+			{
+				return ECSDResult::Error;
+			}
+			const TArray<TArray<uint8>>& dcrSPSs = dcr.GetSequenceParameterSets();
+			for(int32 i=0; i<dcrSPSs.Num(); ++i)
+			{
+				if (!ensure(ElectraDecodersUtil::MPEG::H264::ParseSequenceParameterSet(SPSs, dcrSPSs[i].GetData(), dcrSPSs[i].Num())))
+				{
+					return ECSDResult::Error;
+				}
+			}
+			CurrentDecoderConfiguration = MoveTemp(ConfigData);
+			return SPSs.IsEmpty() ? ECSDResult::NoCSD : ECSDResult::Changed;
+		}
+		// See if pre-extracted CSD is given.
+		ConfigData = ElectraDecodersUtil::GetVariantValueUInt8Array(InFromMap, CSDName);
+		if (ConfigData.Num())
+		{
+			// If the data is the same then we don't need to reparse it.
+			if (ConfigData == CurrentDecoderConfiguration)
+			{
+				return ECSDResult::Unchanged;
+			}
+
+			SPSs.Empty();
+			TArray<ElectraDecodersUtil::MPEG::H264::FNaluInfo> NALUs;
+			if (!ensure(ElectraDecodersUtil::MPEG::H264::ParseBitstreamForNALUs(NALUs, ConfigData.GetData(), ConfigData.Num())))
+			{
+				return ECSDResult::Error;
+			}
+			for(int32 i=0; i<NALUs.Num(); ++i)
+			{
+				if (NALUs[i].Type != 7)
+				{
+					continue;
+				}
+				if (!ensure(ElectraDecodersUtil::MPEG::H264::ParseSequenceParameterSet(SPSs, ConfigData.GetData() + NALUs[i].Offset + NALUs[i].UnitLength, NALUs[i].Size)))
+				{
+					return ECSDResult::Error;
+				}
+			}
+			CurrentDecoderConfiguration = MoveTemp(ConfigData);
+			return SPSs.IsEmpty() ? ECSDResult::NoCSD : ECSDResult::Changed;
+		}
+		return ECSDResult::NoCSD;
+	}
+
+	bool UpdateColorimetry(uint8 colour_primaries, uint8 transfer_characteristics, uint8 matrix_coeffs, uint8 video_full_range_flag, uint8 video_format)
+	{
+		if (!CurrentColorimetry.IsValid() ||
+			CurrentColorimetry->colour_primaries != colour_primaries ||
+			CurrentColorimetry->transfer_characteristics != transfer_characteristics ||
+			CurrentColorimetry->matrix_coeffs != matrix_coeffs ||
+			CurrentColorimetry->video_full_range_flag != video_full_range_flag ||
+			CurrentColorimetry->video_format != video_format)
+		{
+			CurrentColorimetry = MakeShared<ElectraDecodersUtil::MPEG::FCommonColorimetry, ESPMode::ThreadSafe>();
+			CurrentColorimetry->colour_primaries = colour_primaries;
+			CurrentColorimetry->transfer_characteristics = transfer_characteristics;
+			CurrentColorimetry->matrix_coeffs = matrix_coeffs;
+			CurrentColorimetry->video_full_range_flag = video_full_range_flag;
+			CurrentColorimetry->video_format = video_format;
+			return true;
+		}
+		return false;
+	}
+
+	bool HandlePicTiming(const ElectraDecodersUtil::MPEG::FSEIMessage& InSEI, const TMap<uint32, ElectraDecodersUtil::MPEG::H264::FSequenceParameterSet>& InSPSs);
+
+	void Clear()
+	{
+		CurrentDecoderConfiguration.Empty();
+		SPSs.Empty();
+		CurrentBSI.Reset();
+		CurrentColorimetry.Reset();
+		LastErrorMessage.Empty();
+		FMemory::Memzero(ClockTimestamp);
+	}
+
+	static FString AvccBoxName;
+	static FString DCRName;
+	static FString CSDName;
+
+	TArray<uint8> CurrentDecoderConfiguration;
+	TMap<uint32, ElectraDecodersUtil::MPEG::H264::FSequenceParameterSet> SPSs;
+	TSharedPtr<IElectraDecoderBitstreamInfo, ESPMode::ThreadSafe> CurrentBSI;
+	TSharedPtr<ElectraDecodersUtil::MPEG::FCommonColorimetry, ESPMode::ThreadSafe> CurrentColorimetry;
+	ElectraDecodersUtil::MPEG::FCommonPictureTiming ClockTimestamp[3];
+	FString LastErrorMessage;
+
+	bool bReplaceLengthWithStartcode = false;
+};
+FString FElectraDecoderBitstreamProcessorH264::FImpl::AvccBoxName(TEXT("$avcC_box"));
+FString FElectraDecoderBitstreamProcessorH264::FImpl::DCRName(TEXT("dcr"));
+FString FElectraDecoderBitstreamProcessorH264::FImpl::CSDName(TEXT("csd"));
+
+
+FElectraDecoderBitstreamProcessorH264::~FElectraDecoderBitstreamProcessorH264()
+{
+}
+
+FElectraDecoderBitstreamProcessorH264::FElectraDecoderBitstreamProcessorH264(const TMap<FString, FVariant>& InDecoderParams, const TMap<FString, FVariant>& InFormatParams)
+{
+	Impl = MakePimpl<FImpl>();
+
+	int32 S2L = (int32)ElectraDecodersUtil::GetVariantValueSafeI64(InDecoderParams, IElectraDecoderFeature::StartcodeToLength, -1);
+	check(S2L == -1 || S2L == 0);
+
+	Impl->bReplaceLengthWithStartcode = S2L == -1;
+}
+
+bool FElectraDecoderBitstreamProcessorH264::WillModifyBitstreamInPlace()
+{
+	return Impl->bReplaceLengthWithStartcode;
+}
+
+void FElectraDecoderBitstreamProcessorH264::Clear()
+{
+	Impl->Clear();
+}
+
+IElectraDecoderBitstreamProcessor::EProcessResult FElectraDecoderBitstreamProcessorH264::GetCSDFromConfigurationRecord(TArray<uint8>& OutCSD, const TMap<FString, FVariant>& InParamsWithDCRorCSD)
+{
+check(!"TODO");
+	OutCSD.Empty();
+	return EProcessResult::Ok;
+}
+
+IElectraDecoderBitstreamProcessor::EProcessResult FElectraDecoderBitstreamProcessorH264::ProcessInputForDecoding(TSharedPtr<IElectraDecoderBitstreamInfo, ESPMode::ThreadSafe>& OutBSI, FElectraDecoderInputAccessUnit& InOutAccessUnit, const TMap<FString, FVariant>& InAccessUnitSidebandData)
+{
+	// Already processed?
+	if ((InOutAccessUnit.Flags & EElectraDecoderFlags::InputIsProcessed) != EElectraDecoderFlags::None)
+	{
+		return EProcessResult::Ok;
+	}
+	// Set to processed even if we fail somewhere now.
+	InOutAccessUnit.Flags |= EElectraDecoderFlags::InputIsProcessed;
+
+	FImpl::ECSDResult CsdResult = Impl->ExtractSPS(InAccessUnitSidebandData);
+	if (CsdResult == FImpl::ECSDResult::Error)
+	{
+		Impl->LastErrorMessage = FString::Printf(TEXT("Failed to parse codec specific data"));
+		return EProcessResult::Error;
+	}
+	else if (CsdResult == FImpl::ECSDResult::Changed)
+	{
+		TUniquePtr<FImpl::FBitstreamInfo> BSI = MakeUnique<FImpl::FBitstreamInfo>();
+		BSI->SPSs = Impl->SPSs;
+		Impl->CurrentBSI = MakeShareable<IElectraDecoderBitstreamInfo>(BSI.Release());
+	}
+
+	// Assume this is not a sync sample and that it is discardable.
+	// We update the flags in the loop below to reflect the actual states.
+	InOutAccessUnit.Flags &= ~EElectraDecoderFlags::IsSyncSample;
+	InOutAccessUnit.Flags |= EElectraDecoderFlags::IsDiscardable;
+
+	TArray<ElectraDecodersUtil::MPEG::FSEIMessage> SEIMessages;
+
+	// We now go over the data and replace the NALU sizes with startcode (unless the decoder wants to keep the sizes).
+	uint32* NALU = (uint32*)InOutAccessUnit.Data;
+	uint32* LastNALU = (uint32*)ElectraDecodersUtil::AdvancePointer(NALU, InOutAccessUnit.DataSize);
+	while(NALU < LastNALU)
+	{
+		uint32 naluLen = ElectraDecodersUtil::GetFromBigEndian(*NALU);
+
+		// Check the nal_ref_idc in the NAL unit for dependencies.
+		uint8 nal = *(const uint8*)(NALU + 1);
+		check((nal & 0x80) == 0);
+		if ((nal >> 5) != 0)
+		{
+			InOutAccessUnit.Flags &= ~EElectraDecoderFlags::IsDiscardable;
+		}
+		// IDR frame?
+		if ((nal & 0x1f) == 5)
+		{
+			InOutAccessUnit.Flags |= EElectraDecoderFlags::IsSyncSample;
+		}
+		// SEI message(s)?
+		if ((nal & 0x1f) == 6)
+		{
+			ElectraDecodersUtil::MPEG::ExtractSEIMessages(SEIMessages, ElectraDecodersUtil::AdvancePointer(NALU, 5), naluLen-1, ElectraDecodersUtil::MPEG::ESEIStreamType::H264, false);
+		}
+
+		// NOTE: If we want to handle avc3 with inband SPS's we can extract them here!
+
+		if (Impl->bReplaceLengthWithStartcode)
+		{
+			*NALU = ElectraDecodersUtil::GetFromBigEndian(0x00000001U);
+		}
+		NALU = ElectraDecodersUtil::AdvancePointer(NALU, naluLen + 4);
+	}
+	// Narrow down the SEI messages to those we will ultimately handle in SetPropertiesOnOutput() to
+	// avoid creating unnecessary individual bitstream info structures
+	SEIMessages.RemoveAll([](const ElectraDecodersUtil::MPEG::FSEIMessage& m)
+	{
+		return m.PayloadType != ElectraDecodersUtil::MPEG::FSEIMessage::EPayloadType::PT_pic_timing;
+	});
+	// Are there SEI messages we need to store with this access unit?
+	if (SEIMessages.Num())
+	{
+		TUniquePtr<FImpl::FBitstreamInfo> BSI = MakeUnique<FImpl::FBitstreamInfo>();
+		BSI->SPSs = Impl->SPSs;
+		BSI->SEIMessages = MoveTemp(SEIMessages);
+		OutBSI = MakeShareable<IElectraDecoderBitstreamInfo>(BSI.Release());
+	}
+	else
+	{
+		OutBSI = Impl->CurrentBSI;
+	}
+	return CsdResult == FImpl::ECSDResult::Changed ? EProcessResult::CSDChanged : EProcessResult::Ok;
+}
+
+void FElectraDecoderBitstreamProcessorH264::SetPropertiesOnOutput(TMap<FString, FVariant>& InOutProperties, TSharedPtr<IElectraDecoderBitstreamInfo, ESPMode::ThreadSafe> InBSI)
+{
+	if (!InBSI.IsValid())
+	{
+		return;
+	}
+	const FImpl::FBitstreamInfo& BSI(*static_cast<FImpl::FBitstreamInfo*>(InBSI.Get()));
+
+	// We do not know which SPS the decoded slices were referencing, so we just look at the first SPS we have.
+	if (BSI.SPSs.Num())
+	{
+		auto It = BSI.SPSs.CreateConstIterator();
+		const ElectraDecodersUtil::MPEG::H264::FSequenceParameterSet& sps(It->Value);
+		// Set the bit depth and the colorimetry.
+		uint8 colour_primaries=2, transfer_characteristics=2, matrix_coeffs=2;
+		uint8 video_full_range_flag=0, video_format=5;
+		if (sps.colour_description_present_flag)
+		{
+			colour_primaries = sps.colour_primaries;
+			transfer_characteristics = sps.transfer_characteristics;
+			matrix_coeffs = sps.matrix_coefficients;
+		}
+		if (sps.video_signal_type_present_flag)
+		{
+			video_full_range_flag = sps.video_full_range_flag;
+			video_format = sps.video_format;
+		}
+		// Update the colorimetry information. If it changed we update it in the output properties map.
+		if (Impl->UpdateColorimetry(colour_primaries, transfer_characteristics, matrix_coeffs, video_full_range_flag, video_format))
+		{
+			TConstArrayView<uint8> Colorimetry = MakeConstArrayView(reinterpret_cast<const uint8*>(Impl->CurrentColorimetry.Get()), sizeof(ElectraDecodersUtil::MPEG::FCommonColorimetry));
+			InOutProperties.Emplace(IElectraDecoderBitstreamProcessorInfo::CommonColorimetry, TArray<uint8>(Colorimetry));
+		}
+	}
+	uint8 num_bits = 8;
+	InOutProperties.Emplace(IElectraDecoderBitstreamProcessorInfo::NumBitsLuma, FVariant(num_bits));
+
+	// Handle SEI messages
+	for(int32 i=0, iMax=BSI.SEIMessages.Num(); i<iMax; ++i)
+	{
+		const ElectraDecodersUtil::MPEG::FSEIMessage& sei = BSI.SEIMessages[i];
+		switch(sei.PayloadType)
+		{
+			case ElectraDecodersUtil::MPEG::FSEIMessage::EPayloadType::PT_pic_timing:
+			{
+				if (Impl->HandlePicTiming(sei, BSI.SPSs))
+				{
+					// Set from the first clock since we are only dealing with progressive frames, not interlaced fields.
+					TConstArrayView<uint8> PicTiming = MakeConstArrayView(reinterpret_cast<const uint8*>(&Impl->ClockTimestamp[0]), sizeof(ElectraDecodersUtil::MPEG::FCommonPictureTiming));
+					InOutProperties.Emplace(IElectraDecoderBitstreamProcessorInfo::CommonPictureTiming, TArray<uint8>(PicTiming));
+				}
+				break;
+			}
+			default:
+			{
+				break;
+			}
+		}
+	}
+}
+
+FString FElectraDecoderBitstreamProcessorH264::GetLastError()
+{
+	return Impl->LastErrorMessage;
+}
+
+
+bool FElectraDecoderBitstreamProcessorH264::FImpl::HandlePicTiming(const ElectraDecodersUtil::MPEG::FSEIMessage& InSEI, const TMap<uint32, ElectraDecodersUtil::MPEG::H264::FSequenceParameterSet>& InSPSs)
+{
+	// Parsing the pic_timing() SEI message requires the active SPS.
+	// NOTE: We do NOT know which SPS is active, if there are several as the activation is determined by the slice being decoded.
+	if (InSPSs.Num() != 1)
+	{
+		return false;
+	}
+	auto It = InSPSs.CreateConstIterator();
+	const ElectraDecodersUtil::MPEG::H264::FSequenceParameterSet& sps(It->Value);
+
+	ElectraDecodersUtil::MPEG::H264::FBitstreamReader br(InSEI.Message.GetData(), InSEI.Message.Num());
+
+	const bool CpbDpbDelaysPresentFlag = sps.nal_hrd_parameters_present_flag || sps.vcl_hrd_parameters_present_flag;
+	if (CpbDpbDelaysPresentFlag)
+	{
+		int32 NumBits1 = 1 + (sps.nal_hrd_parameters_present_flag ? sps.nal_hrd_parameters.cpb_removal_delay_length_minus1 : sps.vcl_hrd_parameters.cpb_removal_delay_length_minus1);
+		int32 NumBits2 = 1 + (sps.nal_hrd_parameters_present_flag ? sps.nal_hrd_parameters.dpb_output_delay_length_minus1 : sps.vcl_hrd_parameters.dpb_output_delay_length_minus1);
+		uint32 cpb_removal_delay = br.GetBits(NumBits1);
+		uint32 dpb_output_delay = br.GetBits(NumBits2);
+		(void)cpb_removal_delay;
+		(void)dpb_output_delay;
+	}
+	if (sps.pic_struct_present_flag)
+	{
+		static const uint8 kNumClockTS[9] = { 1, 1, 1, 2, 2, 3, 3, 2, 3 };
+		uint32 pic_struct = br.GetBits(4);
+		if (pic_struct > 8)
+		{
+			return false;
+		}
+		for(int32 nc=0,ncMax=kNumClockTS[pic_struct]; nc<ncMax; ++nc)
+		{
+			ElectraDecodersUtil::MPEG::FCommonPictureTiming& ct(ClockTimestamp[nc]);
+			ct.FromH26x = 4;
+			ct.clock_timestamp_flag = (uint8) br.GetBits(1);
+			if (ct.clock_timestamp_flag)
+			{
+				// Set the timing values from the SPS.
+				ct.timing_info_present_flag = sps.timing_info_present_flag;
+				ct.num_units_in_tick = sps.num_units_in_tick;
+				ct.time_scale = sps.time_scale;
+
+				// Read the values from the message.
+				ct.ct_type = (uint8) br.GetBits(2);
+				ct.nuit_field_based_flag = (uint8) br.GetBits(1);
+				ct.counting_type = (uint8) br.GetBits(5);
+				ct.full_timestamp_flag = (uint8) br.GetBits(1);
+				ct.discontinuity_flag = (uint8) br.GetBits(1);
+				ct.cnt_dropped_flag = (uint8) br.GetBits(1);
+				ct.n_frames = (uint16) br.GetBits(8);
+				if (ct.full_timestamp_flag)
+				{
+					ct.seconds_value = (uint8) br.GetBits(6);
+					ct.minutes_value = (uint8) br.GetBits(6);
+					ct.hours_value = (uint8) br.GetBits(5);
+				}
+				else
+				{
+					// seconds_flag
+					if (br.GetBits(1))
+					{
+						ct.seconds_value = (uint8) br.GetBits(6);
+						// minutes_flag
+						if (br.GetBits(1))
+						{
+							ct.minutes_value = (uint8) br.GetBits(6);
+							// hours_flag
+							if (br.GetBits(1))
+							{
+								ct.hours_value = (uint8) br.GetBits(5);
+							}
+						}
+					}
+				}
+				uint32 time_offset_length = sps.nal_hrd_parameters_present_flag ? sps.nal_hrd_parameters.time_offset_length : sps.vcl_hrd_parameters_present_flag ? sps.vcl_hrd_parameters.time_offset_length : 24;
+				uint32 time_offset = br.GetBits(time_offset_length);
+				ct.time_offset = (int32)(time_offset << (32 - time_offset_length)) >> (32 - time_offset_length);
+
+				if (ct.timing_info_present_flag)
+				{
+					ct.clockTimestamp = (((int64)ct.hours_value * 60 + ct.minutes_value) * 60 + ct.seconds_value) * ct.time_scale + ct.n_frames * (ct.num_units_in_tick * (ct.nuit_field_based_flag + 1)) + ct.time_offset;
+				}
+			}
+		}
+	}
+	return true;
+}
diff --git a/Engine/Plugins/Media/ElectraCodecs/Source/ElectraDecoders/Private/Utils/MPEG/ElectraBitstreamProcessor_H265.cpp b/Engine/Plugins/Media/ElectraCodecs/Source/ElectraDecoders/Private/Utils/MPEG/ElectraBitstreamProcessor_H265.cpp
new file mode 100644
--- /dev/null
+++ b/Engine/Plugins/Media/ElectraCodecs/Source/ElectraDecoders/Private/Utils/MPEG/ElectraBitstreamProcessor_H265.cpp
@@ -0,0 +1,549 @@
+// Copyright Epic Games, Inc. All Rights Reserved.
+
+#include "Utils/MPEG/ElectraBitstreamProcessor_H265.h"
+#include "Utils/MPEG/ElectraUtilsMPEGVideo_H265.h"
+#include "Utils/MPEG/ElectraUtilsMPEGVideo.h"
+#include "Utils/MPEG/ElectraUtilsMP4.h"
+#include "IElectraDecoderFeaturesAndOptions.h"
+#include "ElectraDecodersUtils.h"
+
+
+class FElectraDecoderBitstreamProcessorH265::FImpl
+{
+public:
+	class FBitstreamInfo : public IElectraDecoderBitstreamInfo
+	{
+	public:
+		virtual ~FBitstreamInfo() = default;
+		TArray<ElectraDecodersUtil::MPEG::FSEIMessage> PrefixSEIMessages;
+		TArray<ElectraDecodersUtil::MPEG::FSEIMessage> SuffixSEIMessages;
+		TMap<uint32, ElectraDecodersUtil::MPEG::H265::FVideoParameterSet> VPSs;
+		TMap<uint32, ElectraDecodersUtil::MPEG::H265::FSequenceParameterSet> SPSs;
+	};
+
+	enum class ECSDResult
+	{
+		NoCSD,
+		Unchanged,
+		Changed,
+		Error
+	};
+
+	ECSDResult ExtractVPSSPSPrefixSuffixNUTs(const TMap<FString, FVariant>& InFromMap)
+	{
+		// First try the `$hvcC_box`
+		TArray<uint8> ConfigData = ElectraDecodersUtil::GetVariantValueUInt8Array(InFromMap, HvccBoxName);
+		// If empty try the decoder confiuration record, which is the same as the ConfigData box.
+		if (ConfigData.IsEmpty())
+		{
+			ConfigData = ElectraDecodersUtil::GetVariantValueUInt8Array(InFromMap, DCRName);
+		}
+		if (ConfigData.Num())
+		{
+			// If the data is the same then we don't need to reparse it.
+			if (ConfigData == CurrentDecoderConfiguration)
+			{
+				return ECSDResult::Unchanged;
+			}
+			VPSs.Empty();
+			SPSs.Empty();
+			PrefixSEIMessages.Empty();
+			SuffixSEIMessages.Empty();
+			ElectraDecodersUtil::MPEG::H265::FHEVCDecoderConfigurationRecord dcr;
+			if (!dcr.Parse(ConfigData) || dcr.GetSequenceParameterSets().IsEmpty())
+			{
+				return ECSDResult::Error;
+			}
+			const TArray<TArray<uint8>>& dcrVPSs = dcr.GetVideoParameterSets();
+			for(int32 i=0; i<dcrVPSs.Num(); ++i)
+			{
+				if (!ensure(ElectraDecodersUtil::MPEG::H265::ParseVideoParameterSet(VPSs, dcrVPSs[i].GetData(), dcrVPSs[i].Num())))
+				{
+					return ECSDResult::Error;
+				}
+			}
+			const TArray<TArray<uint8>>& dcrSPSs = dcr.GetSequenceParameterSets();
+			for(int32 i=0; i<dcrSPSs.Num(); ++i)
+			{
+				if (!ensure(ElectraDecodersUtil::MPEG::H265::ParseSequenceParameterSet(SPSs, dcrSPSs[i].GetData(), dcrSPSs[i].Num())))
+				{
+					return ECSDResult::Error;
+				}
+			}
+			const TArray<TArray<uint8>>& dcrPrfx = dcr.GetPrefixNUTs();
+			for(int32 i=0; i<dcrPrfx.Num(); ++i)
+			{
+				if (!ensure(ElectraDecodersUtil::MPEG::ExtractSEIMessages(PrefixSEIMessages, dcrPrfx[i].GetData()+2, dcrPrfx[i].Num()-2, ElectraDecodersUtil::MPEG::ESEIStreamType::H265, true)))
+				{
+					return ECSDResult::Error;
+				}
+			}
+			const TArray<TArray<uint8>>& dcrSufx = dcr.GetSuffixNUTs();
+			for(int32 i=0; i<dcrSufx.Num(); ++i)
+			{
+				if (!ensure(ElectraDecodersUtil::MPEG::ExtractSEIMessages(SuffixSEIMessages, dcrSufx[i].GetData()+2, dcrSufx[i].Num()-2, ElectraDecodersUtil::MPEG::ESEIStreamType::H265, false)))
+				{
+					return ECSDResult::Error;
+				}
+			}
+			CurrentDecoderConfiguration = MoveTemp(ConfigData);
+			return SPSs.IsEmpty() ? ECSDResult::NoCSD : ECSDResult::Changed;
+		}
+		// See if pre-extracted CSD is given.
+		ConfigData = ElectraDecodersUtil::GetVariantValueUInt8Array(InFromMap, CSDName);
+		if (ConfigData.Num())
+		{
+			// If the data is the same then we don't need to reparse it.
+			if (ConfigData == CurrentDecoderConfiguration)
+			{
+				return ECSDResult::Unchanged;
+			}
+
+			VPSs.Empty();
+			SPSs.Empty();
+			PrefixSEIMessages.Empty();
+			SuffixSEIMessages.Empty();
+			TArray<ElectraDecodersUtil::MPEG::H265::FNaluInfo> NALUs;
+			if (!ensure(ElectraDecodersUtil::MPEG::H265::ParseBitstreamForNALUs(NALUs, ConfigData.GetData(), ConfigData.Num())))
+			{
+				return ECSDResult::Error;
+			}
+			for(int32 i=0; i<NALUs.Num(); ++i)
+			{
+				// VPS NUT?
+				if (NALUs[i].Type == 32 && !ensure(ElectraDecodersUtil::MPEG::H265::ParseVideoParameterSet(VPSs, ConfigData.GetData() + NALUs[i].Offset + NALUs[i].UnitLength, NALUs[i].Size)))
+				{
+					return ECSDResult::Error;
+				}
+				// SPS NUT?
+				else if (NALUs[i].Type == 33 && !ensure(ElectraDecodersUtil::MPEG::H265::ParseSequenceParameterSet(SPSs, ConfigData.GetData() + NALUs[i].Offset + NALUs[i].UnitLength, NALUs[i].Size)))
+				{
+					return ECSDResult::Error;
+				}
+				// Prefix or suffix NUT?
+				else if (NALUs[i].Type == 39 || NALUs[i].Type == 40)
+				{
+					ElectraDecodersUtil::MPEG::ExtractSEIMessages(NALUs[i].Type == 39 ? PrefixSEIMessages : SuffixSEIMessages, ConfigData.GetData() + NALUs[i].Offset + NALUs[i].UnitLength + 2, NALUs[i].Size - 2, ElectraDecodersUtil::MPEG::ESEIStreamType::H265, NALUs[i].Type == 39);
+				}
+
+			}
+			CurrentDecoderConfiguration = MoveTemp(ConfigData);
+			return SPSs.IsEmpty() ? ECSDResult::NoCSD : ECSDResult::Changed;
+		}
+		return ECSDResult::NoCSD;
+	}
+
+	bool UpdateColorimetry(uint8 colour_primaries, uint8 transfer_characteristics, uint8 matrix_coeffs, uint8 video_full_range_flag, uint8 video_format)
+	{
+		if (!CurrentColorimetry.IsValid() ||
+			CurrentColorimetry->colour_primaries != colour_primaries ||
+			CurrentColorimetry->transfer_characteristics != transfer_characteristics ||
+			CurrentColorimetry->matrix_coeffs != matrix_coeffs ||
+			CurrentColorimetry->video_full_range_flag != video_full_range_flag ||
+			CurrentColorimetry->video_format != video_format)
+		{
+			CurrentColorimetry = MakeShared<ElectraDecodersUtil::MPEG::FCommonColorimetry, ESPMode::ThreadSafe>();
+			CurrentColorimetry->colour_primaries = colour_primaries;
+			CurrentColorimetry->transfer_characteristics = transfer_characteristics;
+			CurrentColorimetry->matrix_coeffs = matrix_coeffs;
+			CurrentColorimetry->video_full_range_flag = video_full_range_flag;
+			CurrentColorimetry->video_format = video_format;
+			return true;
+		}
+		return false;
+	}
+
+	bool HandleTimeCode(const ElectraDecodersUtil::MPEG::FSEIMessage& InSEI, const TMap<uint32, ElectraDecodersUtil::MPEG::H265::FVideoParameterSet>& InVPSs, const TMap<uint32, ElectraDecodersUtil::MPEG::H265::FSequenceParameterSet>& InSPSs);
+
+	void Clear()
+	{
+		CurrentDecoderConfiguration.Empty();
+		SPSs.Empty();
+		PrefixSEIMessages.Empty();
+		SuffixSEIMessages.Empty();
+		CurrentBSI.Reset();
+		CurrentColorimetry.Reset();
+		CurrentMDCV.Reset();
+		CurrentCLLI.Reset();
+		CurrentALTC.Reset();
+		LastErrorMessage.Empty();
+		FMemory::Memzero(ClockTimestamp);
+	}
+
+	static FString HvccBoxName;
+	static FString DCRName;
+	static FString CSDName;
+
+	TArray<uint8> CurrentDecoderConfiguration;
+	TMap<uint32, ElectraDecodersUtil::MPEG::H265::FVideoParameterSet> VPSs;
+	TMap<uint32, ElectraDecodersUtil::MPEG::H265::FSequenceParameterSet> SPSs;
+	TArray<ElectraDecodersUtil::MPEG::FSEIMessage> PrefixSEIMessages;
+	TArray<ElectraDecodersUtil::MPEG::FSEIMessage> SuffixSEIMessages;
+	TSharedPtr<IElectraDecoderBitstreamInfo, ESPMode::ThreadSafe> CurrentBSI;
+	TSharedPtr<ElectraDecodersUtil::MPEG::FCommonColorimetry, ESPMode::ThreadSafe> CurrentColorimetry;
+	TSharedPtr<ElectraDecodersUtil::MPEG::FSEImastering_display_colour_volume, ESPMode::ThreadSafe> CurrentMDCV;
+	TSharedPtr<ElectraDecodersUtil::MPEG::FSEIcontent_light_level_info, ESPMode::ThreadSafe> CurrentCLLI;
+	TSharedPtr<ElectraDecodersUtil::MPEG::FSEIalternative_transfer_characteristics, ESPMode::ThreadSafe> CurrentALTC;
+
+	ElectraDecodersUtil::MPEG::FCommonPictureTiming ClockTimestamp[3];
+	FString LastErrorMessage;
+
+	bool bReplaceLengthWithStartcode = false;
+};
+FString FElectraDecoderBitstreamProcessorH265::FImpl::HvccBoxName(TEXT("$hvcC_box"));
+FString FElectraDecoderBitstreamProcessorH265::FImpl::DCRName(TEXT("dcr"));
+FString FElectraDecoderBitstreamProcessorH265::FImpl::CSDName(TEXT("csd"));
+
+
+FElectraDecoderBitstreamProcessorH265::~FElectraDecoderBitstreamProcessorH265()
+{
+}
+
+FElectraDecoderBitstreamProcessorH265::FElectraDecoderBitstreamProcessorH265(const TMap<FString, FVariant>& InDecoderParams, const TMap<FString, FVariant>& InFormatParams)
+{
+	Impl = MakePimpl<FImpl>();
+
+	int32 S2L = (int32)ElectraDecodersUtil::GetVariantValueSafeI64(InDecoderParams, IElectraDecoderFeature::StartcodeToLength, -1);
+	check(S2L == -1 || S2L == 0);
+
+	Impl->bReplaceLengthWithStartcode = S2L == -1;
+}
+
+bool FElectraDecoderBitstreamProcessorH265::WillModifyBitstreamInPlace()
+{
+	return Impl->bReplaceLengthWithStartcode;
+}
+
+void FElectraDecoderBitstreamProcessorH265::Clear()
+{
+	Impl->Clear();
+}
+
+IElectraDecoderBitstreamProcessor::EProcessResult FElectraDecoderBitstreamProcessorH265::GetCSDFromConfigurationRecord(TArray<uint8>& OutCSD, const TMap<FString, FVariant>& InParamsWithDCRorCSD)
+{
+check(!"TODO");
+	OutCSD.Empty();
+	return EProcessResult::Ok;
+}
+
+IElectraDecoderBitstreamProcessor::EProcessResult FElectraDecoderBitstreamProcessorH265::ProcessInputForDecoding(TSharedPtr<IElectraDecoderBitstreamInfo, ESPMode::ThreadSafe>& OutBSI, FElectraDecoderInputAccessUnit& InOutAccessUnit, const TMap<FString, FVariant>& InAccessUnitSidebandData)
+{
+	// Already processed?
+	if ((InOutAccessUnit.Flags & EElectraDecoderFlags::InputIsProcessed) != EElectraDecoderFlags::None)
+	{
+		return EProcessResult::Ok;
+	}
+
+	auto ReduceSEIMessages = [](TArray<ElectraDecodersUtil::MPEG::FSEIMessage>& io) -> void
+	{
+		io.RemoveAll([](const ElectraDecodersUtil::MPEG::FSEIMessage& m)
+		{
+			return m.PayloadType != ElectraDecodersUtil::MPEG::FSEIMessage::EPayloadType::PT_pic_timing &&
+				   m.PayloadType != ElectraDecodersUtil::MPEG::FSEIMessage::EPayloadType::PT_mastering_display_colour_volume &&
+				   m.PayloadType != ElectraDecodersUtil::MPEG::FSEIMessage::EPayloadType::PT_content_light_level_info &&
+				   m.PayloadType != ElectraDecodersUtil::MPEG::FSEIMessage::EPayloadType::PT_alternative_transfer_characteristics;
+		});
+	};
+
+
+	// Set to processed even if we fail somewhere now.
+	InOutAccessUnit.Flags |= EElectraDecoderFlags::InputIsProcessed;
+
+	FImpl::ECSDResult CsdResult = Impl->ExtractVPSSPSPrefixSuffixNUTs(InAccessUnitSidebandData);
+	if (CsdResult == FImpl::ECSDResult::Error)
+	{
+		Impl->LastErrorMessage = FString::Printf(TEXT("Failed to parse codec specific data"));
+		return EProcessResult::Error;
+	}
+	else if (CsdResult == FImpl::ECSDResult::Changed)
+	{
+		TUniquePtr<FImpl::FBitstreamInfo> BSI = MakeUnique<FImpl::FBitstreamInfo>();
+		BSI->SPSs = Impl->SPSs;
+		BSI->VPSs = Impl->VPSs;
+		ReduceSEIMessages(Impl->PrefixSEIMessages);
+		ReduceSEIMessages(Impl->SuffixSEIMessages);
+		Impl->CurrentBSI = MakeShareable<IElectraDecoderBitstreamInfo>(BSI.Release());
+	}
+
+	// Assume this is not a sync sample and that it is discardable.
+	// We update the flags in the loop below to reflect the actual states.
+	InOutAccessUnit.Flags &= ~EElectraDecoderFlags::IsSyncSample;
+	InOutAccessUnit.Flags &= ~EElectraDecoderFlags::IsDiscardable;
+
+	TArray<ElectraDecodersUtil::MPEG::FSEIMessage> PrefixSEIMessages;
+	TArray<ElectraDecodersUtil::MPEG::FSEIMessage> SuffixSEIMessages;
+
+	// We now go over the data and replace the NALU sizes with startcode (unless the decoder wants to keep the sizes).
+	uint32* NALU = (uint32*)InOutAccessUnit.Data;
+	uint32* LastNALU = (uint32*)ElectraDecodersUtil::AdvancePointer(NALU, InOutAccessUnit.DataSize);
+	while(NALU < LastNALU)
+	{
+		uint32 naluLen = ElectraDecodersUtil::GetFromBigEndian(*NALU);
+
+		// Check the nal_ref_idc in the NAL unit for dependencies.
+		uint8 nut = *(const uint8*)(NALU + 1);
+		check((nut & 0x80) == 0);
+		nut >>= 1;
+
+		// IDR, CRA or BLA frame?
+		if (nut >= 16 && nut <= 21)
+		{
+			InOutAccessUnit.Flags |= EElectraDecoderFlags::IsSyncSample;
+		}
+		// One of TRAIL_N, TSA_N, STSA_N, RADL_N, RASL_N, RSV_VCL_N10, RSV_VCL_N12 or RSV_VCL_N14 ?
+		else if (nut <= 14 && (nut & 1) == 0)
+		{
+			InOutAccessUnit.Flags |= EElectraDecoderFlags::IsDiscardable;
+		}
+		// Prefix or suffix NUT?
+		else if (nut == 39 || nut == 40)
+		{
+			ElectraDecodersUtil::MPEG::ExtractSEIMessages(nut == 39 ? PrefixSEIMessages : SuffixSEIMessages, ElectraDecodersUtil::AdvancePointer(NALU, 6), naluLen-2, ElectraDecodersUtil::MPEG::ESEIStreamType::H265, nut == 39);
+		}
+
+		// NOTE: If we want to handle hev1 with inband VPS/SPS/PPS we can extract them here!
+
+		if (Impl->bReplaceLengthWithStartcode)
+		{
+			*NALU = ElectraDecodersUtil::GetFromBigEndian(0x00000001U);
+		}
+		NALU = ElectraDecodersUtil::AdvancePointer(NALU, naluLen + 4);
+	}
+
+	// Narrow down the SEI messages to those we will ultimately handle in SetPropertiesOnOutput() to
+	// avoid creating unnecessary individual bitstream info structures
+	ReduceSEIMessages(PrefixSEIMessages);
+	ReduceSEIMessages(SuffixSEIMessages);
+	// Are there SEI messages we need to store with this access unit?
+	if (PrefixSEIMessages.Num() || SuffixSEIMessages.Num())
+	{
+		TUniquePtr<FImpl::FBitstreamInfo> BSI = MakeUnique<FImpl::FBitstreamInfo>();
+		BSI->SPSs = Impl->SPSs;
+		BSI->VPSs = Impl->VPSs;
+		BSI->PrefixSEIMessages = Impl->PrefixSEIMessages;
+		BSI->SuffixSEIMessages = Impl->SuffixSEIMessages;
+		BSI->PrefixSEIMessages.Append(MoveTemp(PrefixSEIMessages));
+		BSI->SuffixSEIMessages.Append(MoveTemp(SuffixSEIMessages));
+		OutBSI = MakeShareable<IElectraDecoderBitstreamInfo>(BSI.Release());
+	}
+	else
+	{
+		OutBSI = Impl->CurrentBSI;
+	}
+	return CsdResult == FImpl::ECSDResult::Changed ? EProcessResult::CSDChanged : EProcessResult::Ok;
+}
+
+void FElectraDecoderBitstreamProcessorH265::SetPropertiesOnOutput(TMap<FString, FVariant>& InOutProperties, TSharedPtr<IElectraDecoderBitstreamInfo, ESPMode::ThreadSafe> InBSI)
+{
+	if (!InBSI.IsValid())
+	{
+		return;
+	}
+	const FImpl::FBitstreamInfo& BSI(*static_cast<FImpl::FBitstreamInfo*>(InBSI.Get()));
+
+	uint8 num_bits = 8;
+
+	// We do not know which SPS the decoded slices were referencing, so we just look at the first SPS we have.
+	if (BSI.SPSs.Num())
+	{
+		auto It = Impl->SPSs.CreateConstIterator();
+		const ElectraDecodersUtil::MPEG::H265::FSequenceParameterSet& sps(It->Value);
+		// Set the bit depth and the colorimetry.
+		uint8 colour_primaries=2, transfer_characteristics=2, matrix_coeffs=2;
+		uint8 video_full_range_flag=0, video_format=5;
+
+		num_bits = sps.bit_depth_luma_minus8 + 8;
+		if (sps.vui_parameters_present_flag)
+		{
+			if (sps.vui_parameters.colour_description_present_flag)
+			{
+				colour_primaries = sps.vui_parameters.colour_primaries;
+				transfer_characteristics = sps.vui_parameters.transfer_characteristics;
+				matrix_coeffs = sps.vui_parameters.matrix_coeffs;
+			}
+			if (sps.vui_parameters.video_signal_type_present_flag)
+			{
+				video_full_range_flag = sps.vui_parameters.video_full_range_flag;
+				video_format = sps.vui_parameters.video_format;
+			}
+		}
+
+		// Update the colorimetry information. If it changed we update it in the output properties map.
+		if (Impl->UpdateColorimetry(colour_primaries, transfer_characteristics, matrix_coeffs, video_full_range_flag, video_format))
+		{
+			TConstArrayView<uint8> Colorimetry = MakeConstArrayView(reinterpret_cast<const uint8*>(Impl->CurrentColorimetry.Get()), sizeof(ElectraDecodersUtil::MPEG::FCommonColorimetry));
+			InOutProperties.Emplace(IElectraDecoderBitstreamProcessorInfo::CommonColorimetry, TArray<uint8>(Colorimetry));
+		}
+	}
+	InOutProperties.Emplace(IElectraDecoderBitstreamProcessorInfo::NumBitsLuma, FVariant(num_bits));
+
+	// Handle prefix SEI messages
+	bool bMDCVchanged = false;
+	bool bCLLIchanged = false;
+	bool bALTCchanged = false;
+	for(int32 i=0, iMax=BSI.PrefixSEIMessages.Num(); i<iMax; ++i)
+	{
+		const ElectraDecodersUtil::MPEG::FSEIMessage& sei = BSI.PrefixSEIMessages[i];
+		switch(sei.PayloadType)
+		{
+			case ElectraDecodersUtil::MPEG::FSEIMessage::EPayloadType::PT_time_code:
+			{
+				if (Impl->HandleTimeCode(sei, BSI.VPSs, BSI.SPSs))
+				{
+					// Set from the first clock since we are only dealing with progressive frames, not interlaced fields.
+					TConstArrayView<uint8> TimeCode = MakeConstArrayView(reinterpret_cast<const uint8*>(&Impl->ClockTimestamp[0]), sizeof(ElectraDecodersUtil::MPEG::FCommonPictureTiming));
+					InOutProperties.Emplace(IElectraDecoderBitstreamProcessorInfo::CommonPictureTiming, TArray<uint8>(TimeCode));
+				}
+				break;
+			}
+			case ElectraDecodersUtil::MPEG::FSEIMessage::EPayloadType::PT_mastering_display_colour_volume:
+			{
+				ElectraDecodersUtil::MPEG::FSEImastering_display_colour_volume mdcv;
+				if (ElectraDecodersUtil::MPEG::ParseFromSEIMessage(mdcv, sei))
+				{
+					if (!Impl->CurrentMDCV.IsValid() || FMemory::Memcmp(Impl->CurrentMDCV.Get(), &mdcv, sizeof(mdcv)))
+					{
+						Impl->CurrentMDCV = MakeShareable<ElectraDecodersUtil::MPEG::FSEImastering_display_colour_volume>(new ElectraDecodersUtil::MPEG::FSEImastering_display_colour_volume(mdcv));
+						bMDCVchanged = true;
+					}
+				}
+				break;
+			}
+			case ElectraDecodersUtil::MPEG::FSEIMessage::EPayloadType::PT_content_light_level_info:
+			{
+				ElectraDecodersUtil::MPEG::FSEIcontent_light_level_info clli;
+				if (ElectraDecodersUtil::MPEG::ParseFromSEIMessage(clli, sei))
+				{
+					if (!Impl->CurrentCLLI.IsValid() || FMemory::Memcmp(Impl->CurrentCLLI.Get(), &clli, sizeof(clli)))
+					{
+						Impl->CurrentCLLI = MakeShareable<ElectraDecodersUtil::MPEG::FSEIcontent_light_level_info>(new ElectraDecodersUtil::MPEG::FSEIcontent_light_level_info(clli));
+						bCLLIchanged = true;
+					}
+				}
+				break;
+			}
+			case ElectraDecodersUtil::MPEG::FSEIMessage::EPayloadType::PT_alternative_transfer_characteristics:
+			{
+				ElectraDecodersUtil::MPEG::FSEIalternative_transfer_characteristics altc;
+				if (ElectraDecodersUtil::MPEG::ParseFromSEIMessage(altc, sei))
+				{
+					if (!Impl->CurrentALTC.IsValid() || FMemory::Memcmp(Impl->CurrentALTC.Get(), &altc, sizeof(altc)))
+					{
+						Impl->CurrentALTC = MakeShareable<ElectraDecodersUtil::MPEG::FSEIalternative_transfer_characteristics>(new ElectraDecodersUtil::MPEG::FSEIalternative_transfer_characteristics(altc));
+						bALTCchanged = true;
+					}
+				}
+				break;
+			}
+			default:
+			{
+				break;
+			}
+		}
+	}
+	// Set the changed parameters in the output
+	if (bMDCVchanged)
+	{
+		TConstArrayView<uint8> Params = MakeConstArrayView(reinterpret_cast<const uint8*>(Impl->CurrentMDCV.Get()), sizeof(ElectraDecodersUtil::MPEG::FSEImastering_display_colour_volume));
+		InOutProperties.Emplace(IElectraDecoderBitstreamProcessorInfo::SeiMasteringDisplayColorVolume, TArray<uint8>(Params));
+	}
+	if (bCLLIchanged)
+	{
+		TConstArrayView<uint8> Params = MakeConstArrayView(reinterpret_cast<const uint8*>(Impl->CurrentCLLI.Get()), sizeof(ElectraDecodersUtil::MPEG::FSEIcontent_light_level_info));
+		InOutProperties.Emplace(IElectraDecoderBitstreamProcessorInfo::SeiContentLightLeveInfo, TArray<uint8>(Params));
+	}
+	if (bALTCchanged)
+	{
+		TConstArrayView<uint8> Params = MakeConstArrayView(reinterpret_cast<const uint8*>(Impl->CurrentALTC.Get()), sizeof(ElectraDecodersUtil::MPEG::FSEIalternative_transfer_characteristics));
+		InOutProperties.Emplace(IElectraDecoderBitstreamProcessorInfo::SeiAlternateTransferCharacteristics, TArray<uint8>(Params));
+	}
+}
+
+FString FElectraDecoderBitstreamProcessorH265::GetLastError()
+{
+	return Impl->LastErrorMessage;
+}
+
+bool FElectraDecoderBitstreamProcessorH265::FImpl::HandleTimeCode(const ElectraDecodersUtil::MPEG::FSEIMessage& InSEI, const TMap<uint32, ElectraDecodersUtil::MPEG::H265::FVideoParameterSet>& InVPSs, const TMap<uint32, ElectraDecodersUtil::MPEG::H265::FSequenceParameterSet>& InSPSs)
+{
+	// Parsing the time_code() SEI message requires the active SPS (and VPS)
+	// NOTE: We do NOT know which SPS is active, if there are several as the activation is determined by the slice being decoded.
+	if (InSPSs.Num() != 1)
+	{
+		return false;
+	}
+	auto It = InSPSs.CreateConstIterator();
+	const ElectraDecodersUtil::MPEG::H265::FSequenceParameterSet& sps(It->Value);
+
+	// We also need to have the VPS that this SPS is referencing.
+	if (!VPSs.Contains(sps.sps_video_parameter_set_id))
+	{
+		return false;
+	}
+	const ElectraDecodersUtil::MPEG::H265::FVideoParameterSet& vps(VPSs[sps.sps_video_parameter_set_id]);
+
+	ElectraDecodersUtil::MPEG::H265::FBitstreamReader br(InSEI.Message.GetData(), InSEI.Message.Num());
+
+	uint32 num_clock_ts = br.GetBits(2);
+	for(uint32 nc=0; nc<num_clock_ts; ++nc)
+	{
+		ElectraDecodersUtil::MPEG::FCommonPictureTiming& ct(ClockTimestamp[nc]);
+		ct.FromH26x = 5;
+		ct.clock_timestamp_flag = (uint8) br.GetBits(1);
+		if (ct.clock_timestamp_flag)
+		{
+			// Set the timing values from the SPS or VPS.
+			ct.timing_info_present_flag = vps.vps_timing_info_present_flag;
+			ct.num_units_in_tick = sps.vui_parameters_present_flag && sps.vui_parameters.vui_timing_info_present_flag ? sps.vui_parameters.vui_num_units_in_tick : vps.vps_timing_info_present_flag ? vps.vps_num_units_in_tick : 0;
+			ct.time_scale = sps.vui_parameters_present_flag && sps.vui_parameters.vui_timing_info_present_flag ? sps.vui_parameters.vui_time_scale : vps.vps_timing_info_present_flag ? vps.vps_time_scale : 1;
+
+			// Read the values from the message.
+			ct.nuit_field_based_flag = (uint8) br.GetBits(1);		// In H.265 this is actually now named correctly `units_field_based_flag` ,...
+			ct.counting_type = (uint8) br.GetBits(5);
+			ct.full_timestamp_flag = (uint8) br.GetBits(1);
+			ct.discontinuity_flag = (uint8) br.GetBits(1);
+			ct.cnt_dropped_flag = (uint8) br.GetBits(1);
+			ct.n_frames = (uint16) br.GetBits(9);
+			if (ct.full_timestamp_flag)
+			{
+				ct.seconds_value = (uint8) br.GetBits(6);
+				ct.minutes_value = (uint8) br.GetBits(6);
+				ct.hours_value = (uint8) br.GetBits(5);
+			}
+			else
+			{
+				// seconds_flag
+				if (br.GetBits(1))
+				{
+					ct.seconds_value = (uint8) br.GetBits(6);
+					// minutes_flag
+					if (br.GetBits(1))
+					{
+						ct.minutes_value = (uint8) br.GetBits(6);
+						// hours_flag
+						if (br.GetBits(1))
+						{
+							ct.hours_value = (uint8) br.GetBits(5);
+						}
+					}
+				}
+			}
+			uint32 time_offset_length = br.GetBits(5);
+			if (time_offset_length)
+			{
+				uint32 time_offset = br.GetBits(time_offset_length);
+				ct.time_offset = (int32)(time_offset << (32 - time_offset_length)) >> (32 - time_offset_length);
+			}
+			else
+			{
+				ct.time_offset = 0;
+			}
+
+			if (ct.timing_info_present_flag)
+			{
+				ct.clockTimestamp = (((int64)ct.hours_value * 60 + ct.minutes_value) * 60 + ct.seconds_value) * ct.time_scale + ct.n_frames * (ct.num_units_in_tick * (ct.nuit_field_based_flag + 1)) + ct.time_offset;
+			}
+		}
+	}
+	return true;
+}
diff --git a/Engine/Plugins/Media/ElectraCodecs/Source/ElectraDecoders/Public/IElectraDecoder.h b/Engine/Plugins/Media/ElectraCodecs/Source/ElectraDecoders/Public/IElectraDecoder.h
--- a/Engine/Plugins/Media/ElectraCodecs/Source/ElectraDecoders/Public/IElectraDecoder.h
+++ b/Engine/Plugins/Media/ElectraCodecs/Source/ElectraDecoders/Public/IElectraDecoder.h
@@ -4,6 +4,7 @@
 
 #include "CoreMinimal.h"
 #include "Templates/SharedPointer.h"
+#include "Templates/PimplPtr.h"
 #include "Containers/Map.h"
 #include "Misc/Timespan.h"
 #include "Misc/Variant.h"
@@ -34,6 +35,7 @@ public:
 };
 
 
+
 enum class EElectraDecoderFlags : uint32
 {
 	None = 0,
@@ -42,28 +44,96 @@ enum class EElectraDecoderFlags : uint32
 	IsReplaySample = 0x04,
 	IsLastReplaySample = 0x08,
 	InitCSDOnly = 0x10,
-	DoNotOutput = 0x20
+	DoNotOutput = 0x20,
+	InputIsProcessed = 0x100
 };
 ENUM_CLASS_FLAGS(EElectraDecoderFlags);
 
-class IElectraDecoder : public TSharedFromThis<IElectraDecoder, ESPMode::ThreadSafe>
+struct FElectraDecoderInputAccessUnit
+{
+	const void* Data = nullptr;
+	int64 DataSize = 0;
+	FTimespan DTS;
+	FTimespan PTS;
+	FTimespan Duration;
+	uint64 UserValue = 0;
+	EElectraDecoderFlags Flags = EElectraDecoderFlags::None;
+};
+
+
 
+class IElectraDecoderBitstreamInfo
 {
 public:
+	virtual ~IElectraDecoderBitstreamInfo()
+	{ }
+};
 
-	virtual ~IElectraDecoder() = default;
 
-	struct FInputAccessUnit
+
+class IElectraDecoderBitstreamProcessor : public TSharedFromThis<IElectraDecoderBitstreamProcessor, ESPMode::ThreadSafe>
+{
+public:
+	virtual ~IElectraDecoderBitstreamProcessor() = default;
+
+	enum class EProcessResult
 	{
-		const void* Data = nullptr;
-		int64 DataSize = 0;
-		FTimespan DTS;
-		FTimespan PTS;
-		FTimespan Duration;
-		uint64 UserValue = 0;
-		EElectraDecoderFlags Flags = EElectraDecoderFlags::None;
+		Ok,
+		CSDChanged,
+		Error
 	};
 
+	/**
+	 * Returns whether or not calling ProcessInputForDecoding() will modify the provided bitstream in place.
+	 * If true you may need to create a temporary copy if you want to retain the bitstream in a cache.
+	 */
+	virtual bool WillModifyBitstreamInPlace() = 0;
+
+	/**
+	 * Clears the processor. Must be called when seeking or flushing.
+	 */
+	virtual void Clear() = 0;
+
+	/**
+	 * Returns codec specific data (CSD) from a decoder configuration record (DCR).
+	 * DCR is provided in a container like ISO/IEC 14496-12 but some decoders need to be given the
+	 * encapsulated CSD. If the DCR cannot be parsed `Error` is returned.
+	 * Otherwise, if there is no DCR present in the parameter map (which _could_ be an error on your end if DCR should be there)
+	 * `Ok` is returned with the (possibly emptry) CSD.
+	 */
+	virtual EProcessResult GetCSDFromConfigurationRecord(TArray<uint8>& OutCSD, const TMap<FString, FVariant>& InParamsWithDCRorCSD) = 0;
+
+	/**
+	 * Processes the bitstream and places sideband information in the provided IElectraDecoderBitstreamInfo for use
+	 * with a subsequent call to SetPropertiesOnOutput().
+	 * If necessary the bitstream is modified in place.
+	 * See `WillModifyBitstreamInPlace()`
+	 * Returns `Ok` or `Error`. If the codec specific data changed `CSDChanged` is returned which is only
+	 * informational for you and is functionally the same as `Ok`.
+	 */
+	virtual EProcessResult ProcessInputForDecoding(TSharedPtr<IElectraDecoderBitstreamInfo, ESPMode::ThreadSafe>& OutBSI, FElectraDecoderInputAccessUnit& InOutAccessUnit, const TMap<FString, FVariant>& InAccessUnitSidebandData) = 0;
+
+	/**
+	 * Sets previously extracted sideband data in the provided output dictionary.
+	 */
+	virtual void SetPropertiesOnOutput(TMap<FString, FVariant>& InOutProperties, TSharedPtr<IElectraDecoderBitstreamInfo, ESPMode::ThreadSafe> InBSI) = 0;
+
+	/**
+	 * Returns an error message if either method returned `Error`.
+	 */
+	virtual FString GetLastError() = 0;
+};
+
+
+
+class IElectraDecoder : public TSharedFromThis<IElectraDecoder, ESPMode::ThreadSafe>
+{
+public:
+
+	virtual ~IElectraDecoder() = default;
+
+	using FInputAccessUnit = FElectraDecoderInputAccessUnit;
+
 	enum class EDecoderError
 	{
 		// No error, all is well.
@@ -268,6 +338,13 @@ public:
 	 */
 	virtual TSharedPtr<IElectraDecoderOutput, ESPMode::ThreadSafe> GetOutput() = 0;
 
+
+	/**
+	 * Creates a processor to pre-parse the decoder input bitstream.
+	 */
+	virtual TSharedPtr<IElectraDecoderBitstreamProcessor, ESPMode::ThreadSafe> CreateBitstreamProcessor() = 0;
+
 	virtual void Suspend() = 0;
 	virtual void Resume() = 0;
 };
+
diff --git a/Engine/Plugins/Media/ElectraCodecs/Source/ElectraDecoders/Public/IElectraDecoderFeaturesAndOptions.h b/Engine/Plugins/Media/ElectraCodecs/Source/ElectraDecoders/Public/IElectraDecoderFeaturesAndOptions.h
--- a/Engine/Plugins/Media/ElectraCodecs/Source/ElectraDecoders/Public/IElectraDecoderFeaturesAndOptions.h
+++ b/Engine/Plugins/Media/ElectraCodecs/Source/ElectraDecoders/Public/IElectraDecoderFeaturesAndOptions.h
@@ -162,7 +162,7 @@ namespace IElectraDecoderOption
 
 
 /**
- * This namespace contains constant names describe an decoder format.
+ * This namespace contains constant names that describe an decoder format.
  */
 namespace IElectraDecoderFormatInfo
 {
@@ -182,3 +182,54 @@ namespace IElectraDecoderFormatInfo
 	 */
 	const TCHAR* const HumanReadableFormatName = TEXT("human_readable_format_name");
 }
+
+
+
+/**
+ * This namespace contains constant names that describe bitstream processor values.
+ */
+namespace IElectraDecoderBitstreamProcessorInfo
+{
+	/**
+	 * FVariant(uint8)
+	 *
+	 * Number of luminance bits.
+	 */
+	const TCHAR* const NumBitsLuma = TEXT("num_bits_luma");
+
+	/**
+	 * FVariant(FCommonColorimetry)
+	 *
+	 * Colorimetry information in standard MPEG definition.
+	 */
+	const TCHAR* const CommonColorimetry = TEXT("common_colorimetry");
+
+	/**
+	 * FVariant(FCommonPictureTiming)
+	 *
+	 * Timing from a bitstream SEI message.
+	 */
+	const TCHAR* const CommonPictureTiming = TEXT("common_pic_timing");
+
+	/**
+	 * FVariant(FSEImastering_display_colour_volume)
+	 *
+	 * Mastering display colour volume from bitstream SEI message
+	 */
+	const TCHAR* const SeiMasteringDisplayColorVolume = TEXT("sei_mdcv");
+
+	/**
+	 * FVariant(FSEIcontent_light_level_info)
+	 *
+	 * Content light level info from bitstream SEI message
+	 */
+	const TCHAR* const SeiContentLightLeveInfo = TEXT("sei_clli");
+
+	/**
+	 * FVariant(FSEIalternative_transfer_characteristics)
+	 *
+	 * Alternate transfer characterisctics from bitstream SEI message
+	 */
+	const TCHAR* const SeiAlternateTransferCharacteristics = TEXT("sei_altc");
+
+}
diff --git a/Engine/Plugins/Media/ElectraCodecs/Source/ElectraDecoders/Public/Utils/ElectraBitstreamProcessorDefault.h b/Engine/Plugins/Media/ElectraCodecs/Source/ElectraDecoders/Public/Utils/ElectraBitstreamProcessorDefault.h
new file mode 100644
--- /dev/null
+++ b/Engine/Plugins/Media/ElectraCodecs/Source/ElectraDecoders/Public/Utils/ElectraBitstreamProcessorDefault.h
@@ -0,0 +1,26 @@
+// Copyright Epic Games, Inc. All Rights Reserved.
+
+#pragma once
+
+#include "IElectraDecoder.h"
+
+
+class FElectraDecoderBitstreamProcessorDefault : public IElectraDecoderBitstreamProcessor
+{
+public:
+	static TSharedPtr<IElectraDecoderBitstreamProcessor, ESPMode::ThreadSafe> Create()
+	{ return MakeShareable<>(new FElectraDecoderBitstreamProcessorDefault); }
+
+	virtual ~FElectraDecoderBitstreamProcessorDefault();
+	bool WillModifyBitstreamInPlace() override;
+	void Clear() override;
+	EProcessResult GetCSDFromConfigurationRecord(TArray<uint8>& OutCSD, const TMap<FString, FVariant>& InParamsWithDCRorCSD) override;
+	EProcessResult ProcessInputForDecoding(TSharedPtr<IElectraDecoderBitstreamInfo, ESPMode::ThreadSafe>& OutBSI, FElectraDecoderInputAccessUnit& InOutAccessUnit, const TMap<FString, FVariant>& InAccessUnitSidebandData) override;
+	void SetPropertiesOnOutput(TMap<FString, FVariant>& InOutProperties, TSharedPtr<IElectraDecoderBitstreamInfo, ESPMode::ThreadSafe> InBSI) override;
+	FString GetLastError() override;
+
+private:
+	ELECTRADECODERS_API FElectraDecoderBitstreamProcessorDefault();
+	class FImpl;
+	TPimplPtr<FImpl> Impl;
+};
diff --git a/Engine/Plugins/Media/ElectraCodecs/Source/ElectraDecoders/Public/Utils/ElectraBitstreamProcessorGenericVideo.h b/Engine/Plugins/Media/ElectraCodecs/Source/ElectraDecoders/Public/Utils/ElectraBitstreamProcessorGenericVideo.h
new file mode 100644
--- /dev/null
+++ b/Engine/Plugins/Media/ElectraCodecs/Source/ElectraDecoders/Public/Utils/ElectraBitstreamProcessorGenericVideo.h
@@ -0,0 +1,26 @@
+// Copyright Epic Games, Inc. All Rights Reserved.
+
+#pragma once
+
+#include "IElectraDecoder.h"
+
+
+class FElectraDecoderBitstreamProcessorGenericVideo : public IElectraDecoderBitstreamProcessor
+{
+public:
+	static TSharedPtr<IElectraDecoderBitstreamProcessor, ESPMode::ThreadSafe> Create(const TMap<FString, FVariant>& InDecoderParams, const TMap<FString, FVariant>& InFormatParams)
+	{ return MakeShareable<>(new FElectraDecoderBitstreamProcessorGenericVideo(InDecoderParams, InFormatParams)); }
+
+	virtual ~FElectraDecoderBitstreamProcessorGenericVideo();
+	bool WillModifyBitstreamInPlace() override;
+	void Clear() override;
+	EProcessResult GetCSDFromConfigurationRecord(TArray<uint8>& OutCSD, const TMap<FString, FVariant>& InParamsWithDCRorCSD) override;
+	EProcessResult ProcessInputForDecoding(TSharedPtr<IElectraDecoderBitstreamInfo, ESPMode::ThreadSafe>& OutBSI, FElectraDecoderInputAccessUnit& InOutAccessUnit, const TMap<FString, FVariant>& InAccessUnitSidebandData) override;
+	void SetPropertiesOnOutput(TMap<FString, FVariant>& InOutProperties, TSharedPtr<IElectraDecoderBitstreamInfo, ESPMode::ThreadSafe> InBSI) override;
+	FString GetLastError() override;
+
+private:
+	ELECTRADECODERS_API FElectraDecoderBitstreamProcessorGenericVideo(const TMap<FString, FVariant>& InDecoderParams, const TMap<FString, FVariant>& InFormatParams);
+	class FImpl;
+	TPimplPtr<FImpl> Impl;
+};
diff --git a/Engine/Plugins/Media/ElectraCodecs/Source/ElectraDecoders/Public/Utils/Google/ElectraBitstreamProcessorVPx.h b/Engine/Plugins/Media/ElectraCodecs/Source/ElectraDecoders/Public/Utils/Google/ElectraBitstreamProcessorVPx.h
new file mode 100644
--- /dev/null
+++ b/Engine/Plugins/Media/ElectraCodecs/Source/ElectraDecoders/Public/Utils/Google/ElectraBitstreamProcessorVPx.h
@@ -0,0 +1,26 @@
+// Copyright Epic Games, Inc. All Rights Reserved.
+
+#pragma once
+
+#include "IElectraDecoder.h"
+
+
+class FElectraDecoderBitstreamProcessorVPx : public IElectraDecoderBitstreamProcessor
+{
+public:
+	static TSharedPtr<IElectraDecoderBitstreamProcessor, ESPMode::ThreadSafe> Create(const TMap<FString, FVariant>& InDecoderParams, const TMap<FString, FVariant>& InFormatParams)
+	{ return MakeShareable<>(new FElectraDecoderBitstreamProcessorVPx(InDecoderParams, InFormatParams)); }
+
+	virtual ~FElectraDecoderBitstreamProcessorVPx();
+	bool WillModifyBitstreamInPlace() override;
+	void Clear() override;
+	EProcessResult GetCSDFromConfigurationRecord(TArray<uint8>& OutCSD, const TMap<FString, FVariant>& InParamsWithDCRorCSD) override;
+	EProcessResult ProcessInputForDecoding(TSharedPtr<IElectraDecoderBitstreamInfo, ESPMode::ThreadSafe>& OutBSI, FElectraDecoderInputAccessUnit& InOutAccessUnit, const TMap<FString, FVariant>& InAccessUnitSidebandData) override;
+	void SetPropertiesOnOutput(TMap<FString, FVariant>& InOutProperties, TSharedPtr<IElectraDecoderBitstreamInfo, ESPMode::ThreadSafe> InBSI) override;
+	FString GetLastError() override;
+
+private:
+	ELECTRADECODERS_API FElectraDecoderBitstreamProcessorVPx(const TMap<FString, FVariant>& InDecoderParams, const TMap<FString, FVariant>& InFormatParams);
+	class FImpl;
+	TPimplPtr<FImpl> Impl;
+};
diff --git a/Engine/Plugins/Media/ElectraCodecs/Source/ElectraDecoders/Public/Utils/MPEG/ElectraBitstreamProcessor_H264.h b/Engine/Plugins/Media/ElectraCodecs/Source/ElectraDecoders/Public/Utils/MPEG/ElectraBitstreamProcessor_H264.h
new file mode 100644
--- /dev/null
+++ b/Engine/Plugins/Media/ElectraCodecs/Source/ElectraDecoders/Public/Utils/MPEG/ElectraBitstreamProcessor_H264.h
@@ -0,0 +1,26 @@
+// Copyright Epic Games, Inc. All Rights Reserved.
+
+#pragma once
+
+#include "IElectraDecoder.h"
+
+
+class FElectraDecoderBitstreamProcessorH264 : public IElectraDecoderBitstreamProcessor
+{
+public:
+	static TSharedPtr<IElectraDecoderBitstreamProcessor, ESPMode::ThreadSafe> Create(const TMap<FString, FVariant>& InDecoderParams, const TMap<FString, FVariant>& InFormatParams)
+	{ return MakeShareable<>(new FElectraDecoderBitstreamProcessorH264(InDecoderParams, InFormatParams)); }
+
+	virtual ~FElectraDecoderBitstreamProcessorH264();
+	bool WillModifyBitstreamInPlace() override;
+	void Clear() override;
+	EProcessResult GetCSDFromConfigurationRecord(TArray<uint8>& OutCSD, const TMap<FString, FVariant>& InParamsWithDCRorCSD) override;
+	EProcessResult ProcessInputForDecoding(TSharedPtr<IElectraDecoderBitstreamInfo, ESPMode::ThreadSafe>& OutBSI, FElectraDecoderInputAccessUnit& InOutAccessUnit, const TMap<FString, FVariant>& InAccessUnitSidebandData) override;
+	void SetPropertiesOnOutput(TMap<FString, FVariant>& InOutProperties, TSharedPtr<IElectraDecoderBitstreamInfo, ESPMode::ThreadSafe> InBSI) override;
+	FString GetLastError() override;
+
+private:
+	ELECTRADECODERS_API FElectraDecoderBitstreamProcessorH264(const TMap<FString, FVariant>& InDecoderParams, const TMap<FString, FVariant>& InFormatParams);
+	class FImpl;
+	TPimplPtr<FImpl> Impl;
+};
diff --git a/Engine/Plugins/Media/ElectraCodecs/Source/ElectraDecoders/Public/Utils/MPEG/ElectraBitstreamProcessor_H265.h b/Engine/Plugins/Media/ElectraCodecs/Source/ElectraDecoders/Public/Utils/MPEG/ElectraBitstreamProcessor_H265.h
new file mode 100644
--- /dev/null
+++ b/Engine/Plugins/Media/ElectraCodecs/Source/ElectraDecoders/Public/Utils/MPEG/ElectraBitstreamProcessor_H265.h
@@ -0,0 +1,26 @@
+// Copyright Epic Games, Inc. All Rights Reserved.
+
+#pragma once
+
+#include "IElectraDecoder.h"
+
+
+class FElectraDecoderBitstreamProcessorH265 : public IElectraDecoderBitstreamProcessor
+{
+public:
+	static TSharedPtr<IElectraDecoderBitstreamProcessor, ESPMode::ThreadSafe> Create(const TMap<FString, FVariant>& InDecoderParams, const TMap<FString, FVariant>& InFormatParams)
+	{ return MakeShareable<>(new FElectraDecoderBitstreamProcessorH265(InDecoderParams, InFormatParams)); }
+
+	virtual ~FElectraDecoderBitstreamProcessorH265();
+	bool WillModifyBitstreamInPlace() override;
+	void Clear() override;
+	EProcessResult GetCSDFromConfigurationRecord(TArray<uint8>& OutCSD, const TMap<FString, FVariant>& InParamsWithDCRorCSD) override;
+	EProcessResult ProcessInputForDecoding(TSharedPtr<IElectraDecoderBitstreamInfo, ESPMode::ThreadSafe>& OutBSI, FElectraDecoderInputAccessUnit& InOutAccessUnit, const TMap<FString, FVariant>& InAccessUnitSidebandData) override;
+	void SetPropertiesOnOutput(TMap<FString, FVariant>& InOutProperties, TSharedPtr<IElectraDecoderBitstreamInfo, ESPMode::ThreadSafe> InBSI) override;
+	FString GetLastError() override;
+
+private:
+	ELECTRADECODERS_API FElectraDecoderBitstreamProcessorH265(const TMap<FString, FVariant>& InDecoderParams, const TMap<FString, FVariant>& InFormatParams);
+	class FImpl;
+	TPimplPtr<FImpl> Impl;
+};
diff --git a/Engine/Plugins/Media/FlacDecoderElectra/Source/FlacDecoderElectra/Private/FlacDecoder/ElectraMediaFlacDecoder.cpp b/Engine/Plugins/Media/FlacDecoderElectra/Source/FlacDecoderElectra/Private/FlacDecoder/ElectraMediaFlacDecoder.cpp
--- a/Engine/Plugins/Media/FlacDecoderElectra/Source/FlacDecoderElectra/Private/FlacDecoder/ElectraMediaFlacDecoder.cpp
+++ b/Engine/Plugins/Media/FlacDecoderElectra/Source/FlacDecoderElectra/Private/FlacDecoder/ElectraMediaFlacDecoder.cpp
@@ -15,6 +15,7 @@
 #include "IElectraDecoderResourceDelegate.h"
 #include "ElectraDecodersUtils.h"
 #include "Utils/MPEG/ElectraUtilsMP4.h"		// for parsing the 'dfLa' box
+#include "Utils/ElectraBitstreamProcessorDefault.h"
 #include "Utils/ElectraBitstreamReader.h"
 
 #include "stream_decoder.h"
@@ -151,6 +152,9 @@ public:
 	EOutputStatus HaveOutput() override;
 	TSharedPtr<IElectraDecoderOutput, ESPMode::ThreadSafe> GetOutput() override;
 
+	TSharedPtr<IElectraDecoderBitstreamProcessor, ESPMode::ThreadSafe> CreateBitstreamProcessor() override
+	{ return FElectraDecoderBitstreamProcessorDefault::Create(); }
+
 	void Suspend() override
 	{ }
 	void Resume() override
diff --git a/Engine/Plugins/Media/HAPDecoderElectra/Source/HAPDecoderElectra/Private/HAPDecoder/ElectraMediaHAPDecoder.cpp b/Engine/Plugins/Media/HAPDecoderElectra/Source/HAPDecoderElectra/Private/HAPDecoder/ElectraMediaHAPDecoder.cpp
--- a/Engine/Plugins/Media/HAPDecoderElectra/Source/HAPDecoderElectra/Private/HAPDecoder/ElectraMediaHAPDecoder.cpp
+++ b/Engine/Plugins/Media/HAPDecoderElectra/Source/HAPDecoderElectra/Private/HAPDecoder/ElectraMediaHAPDecoder.cpp
@@ -14,6 +14,7 @@
 #include "IElectraDecoderOutputVideo.h"
 #include "IElectraDecoderResourceDelegate.h"
 #include "ElectraDecodersUtils.h"
+#include "Utils/ElectraBitstreamProcessorDefault.h"
 #include COMPILED_PLATFORM_HEADER(ElectraDecoderGPUBufferHelpers.h)
 
 #include "hap.h"
@@ -283,6 +284,9 @@ public:
 	EOutputStatus HaveOutput() override;
 	TSharedPtr<IElectraDecoderOutput, ESPMode::ThreadSafe> GetOutput() override;
 
+	TSharedPtr<IElectraDecoderBitstreamProcessor, ESPMode::ThreadSafe> CreateBitstreamProcessor() override
+	{ return FElectraDecoderBitstreamProcessorDefault::Create(); }
+
 	void Suspend() override
 	{ }
 	void Resume() override
diff --git a/Engine/Plugins/Media/OpusDecoderElectra/Source/OpusDecoderElectra/Private/OpusDecoder/ElectraMediaOpusDecoder.cpp b/Engine/Plugins/Media/OpusDecoderElectra/Source/OpusDecoderElectra/Private/OpusDecoder/ElectraMediaOpusDecoder.cpp
--- a/Engine/Plugins/Media/OpusDecoderElectra/Source/OpusDecoderElectra/Private/OpusDecoder/ElectraMediaOpusDecoder.cpp
+++ b/Engine/Plugins/Media/OpusDecoderElectra/Source/OpusDecoderElectra/Private/OpusDecoder/ElectraMediaOpusDecoder.cpp
@@ -14,6 +14,7 @@
 #include "IElectraDecoderOutputAudio.h"
 #include "IElectraDecoderResourceDelegate.h"
 #include "ElectraDecodersUtils.h"
+#include "Utils/ElectraBitstreamProcessorDefault.h"
 #include "Utils/MPEG/ElectraUtilsMP4.h"		// for parsing the 'dOps' box
 
 #include "opus_multistream.h"
@@ -151,6 +152,9 @@ public:
 	EOutputStatus HaveOutput() override;
 	TSharedPtr<IElectraDecoderOutput, ESPMode::ThreadSafe> GetOutput() override;
 
+	TSharedPtr<IElectraDecoderBitstreamProcessor, ESPMode::ThreadSafe> CreateBitstreamProcessor() override
+	{ return FElectraDecoderBitstreamProcessorDefault::Create(); }
+
 	void Suspend() override
 	{ }
 	void Resume() override
diff --git a/Engine/Plugins/Media/VPxDecoderElectra/Source/VPxDecoderElectra/Private/VPxDecoder/ElectraMediaVPxDecoder.cpp b/Engine/Plugins/Media/VPxDecoderElectra/Source/VPxDecoderElectra/Private/VPxDecoder/ElectraMediaVPxDecoder.cpp
--- a/Engine/Plugins/Media/VPxDecoderElectra/Source/VPxDecoderElectra/Private/VPxDecoder/ElectraMediaVPxDecoder.cpp
+++ b/Engine/Plugins/Media/VPxDecoderElectra/Source/VPxDecoderElectra/Private/VPxDecoder/ElectraMediaVPxDecoder.cpp
@@ -15,6 +15,7 @@
 #include "IElectraDecoderResourceDelegate.h"
 #include "ElectraDecodersUtils.h"
 #include "Utils/Google/ElectraUtilsVPxVideo.h"
+#include "Utils/Google/ElectraBitstreamProcessorVPx.h"
 #include "Containers/Queue.h"
 #include <Stats/Stats.h>
 
@@ -225,6 +226,13 @@ public:
 	EOutputStatus HaveOutput() override;
 	TSharedPtr<IElectraDecoderOutput, ESPMode::ThreadSafe> GetOutput() override;
 
+	TSharedPtr<IElectraDecoderBitstreamProcessor, ESPMode::ThreadSafe> CreateBitstreamProcessor() override
+	{
+		TMap<FString, FVariant> DecoderFeatures;
+		GetFeatures(DecoderFeatures);
+		return FElectraDecoderBitstreamProcessorVPx::Create(DecoderFeatures, InitialCreationOptions);
+	}
+
 	void Suspend() override
 	{ }
 	void Resume() override
@@ -274,6 +282,7 @@ private:
 
 	bool ConvertDecodedImageToNV12orP010(FVideoDecoderOutputVPxElectra* InNewOutput, const vpx_image_t* InDecodedImage, void* PlatformDevice, int32 PlatformDeviceVersion) const;
 
+	TMap<FString, FVariant> InitialCreationOptions;
 	TWeakPtr<IElectraDecoderResourceDelegate, ESPMode::ThreadSafe> ResourceDelegate;
 
 	TQueue<TSharedPtr<FDecoderInput, ESPMode::NotThreadSafe>, EQueueMode::Spsc> PendingDecoderInput;
@@ -380,9 +389,10 @@ void FElectraMediaVPxDecoder::Shutdown()
 /*********************************************************************************************************************/
 
 FVideoDecoderVPxElectra::FVideoDecoderVPxElectra(const TMap<FString, FVariant>& InOptions, TSharedPtr<IElectraDecoderResourceDelegate, ESPMode::ThreadSafe> InResourceDelegate)
+	: InitialCreationOptions(InOptions)
+	, ResourceDelegate(InResourceDelegate)
 {
 	Codec4CC = (uint32)ElectraDecodersUtil::GetVariantValueSafeU64(InOptions, TEXT("codec_4cc"), 0);
-	ResourceDelegate = InResourceDelegate;
 
 	MaxWidth = (uint32)ElectraDecodersUtil::GetVariantValueSafeU64(InOptions, TEXT("max_width"), 1920);
 	MaxHeight = (uint32)ElectraDecodersUtil::GetVariantValueSafeU64(InOptions, TEXT("max_height"), 1080);
