Electra: Switching all remaining code to use the new DCR

#rnx

diff --git a/Engine/Plugins/Media/ElectraCodecs/Source/ElectraDecoders/Private/Android/h264/VideoDecoderH264_Android.cpp b/Engine/Plugins/Media/ElectraCodecs/Source/ElectraDecoders/Private/Android/h264/VideoDecoderH264_Android.cpp
--- a/Engine/Plugins/Media/ElectraCodecs/Source/ElectraDecoders/Private/Android/h264/VideoDecoderH264_Android.cpp
+++ b/Engine/Plugins/Media/ElectraCodecs/Source/ElectraDecoders/Private/Android/h264/VideoDecoderH264_Android.cpp
@@ -202,7 +202,7 @@ private:
 	{
 		FInputAccessUnit AccessUnit;
 		TMap<FString, FVariant> AdditionalOptions;
-		TSharedPtr<ElectraDecodersUtil::MPEG::FISO14496_10_seq_parameter_set_data, ESPMode::ThreadSafe> SPS;
+		TSharedPtr<ElectraDecodersUtil::MPEG::H264::FSequenceParameterSet, ESPMode::ThreadSafe> SPS;
 	};
 
 	enum class EDecodeState
@@ -214,8 +214,8 @@ private:
 
 	bool PostError(int32 ApiReturnValue, FString Message, int32 Code);
 
-	TSharedPtr<ElectraDecodersUtil::MPEG::FAVCDecoderConfigurationRecord, ESPMode::ThreadSafe> GetDCRFromOptions(const TMap<FString, FVariant>& InOptions);
-	TSharedPtr<ElectraDecodersUtil::MPEG::FISO14496_10_seq_parameter_set_data, ESPMode::ThreadSafe> GetSPSFromOptions(const TMap<FString, FVariant>& InOptions);
+	TSharedPtr<ElectraDecodersUtil::MPEG::H264::FAVCDecoderConfigurationRecord, ESPMode::ThreadSafe> GetDCRFromOptions(const TMap<FString, FVariant>& InOptions);
+	TSharedPtr<ElectraDecodersUtil::MPEG::H264::FSequenceParameterSet, ESPMode::ThreadSafe> GetSPSFromOptions(const TMap<FString, FVariant>& InOptions);
 
 	bool InternalDecoderCreate();
 	bool InternalHandleDecoderCreate();
@@ -239,8 +239,8 @@ private:
 	TWeakPtr<IElectraDecoderResourceDelegate, ESPMode::ThreadSafe> ResourceDelegate;
 	IElectraDecoderResourceDelegateAndroid::IDecoderPlatformResourceAndroid* PlatformResource = nullptr;
 
-	TSharedPtr<ElectraDecodersUtil::MPEG::FAVCDecoderConfigurationRecord, ESPMode::ThreadSafe> ConfigRecord;
-	TSharedPtr<ElectraDecodersUtil::MPEG::FISO14496_10_seq_parameter_set_data, ESPMode::ThreadSafe> CurrentSPS;
+	TSharedPtr<ElectraDecodersUtil::MPEG::H264::FAVCDecoderConfigurationRecord, ESPMode::ThreadSafe> ConfigRecord;
+	TSharedPtr<ElectraDecodersUtil::MPEG::H264::FSequenceParameterSet, ESPMode::ThreadSafe> CurrentSPS;
 	TSharedPtr<FElectraVideoDecoderOutputH264_Android, ESPMode::ThreadSafe> CurrentOutput;
 	TArray<TSharedPtr<FDecoderInput, ESPMode::NotThreadSafe>> InDecoderInput;
 
@@ -319,15 +319,14 @@ FElectraVideoDecoderH264_Android::~FElectraVideoDecoderH264_Android()
 }
 
 
-TSharedPtr<ElectraDecodersUtil::MPEG::FAVCDecoderConfigurationRecord, ESPMode::ThreadSafe> FElectraVideoDecoderH264_Android::GetDCRFromOptions(const TMap<FString, FVariant>& InOptions)
+TSharedPtr<ElectraDecodersUtil::MPEG::H264::FAVCDecoderConfigurationRecord, ESPMode::ThreadSafe> FElectraVideoDecoderH264_Android::GetDCRFromOptions(const TMap<FString, FVariant>& InOptions)
 {
 	TArray<uint8> SidebandData = ElectraDecodersUtil::GetVariantValueUInt8Array(InOptions, TEXT("dcr"));
 	if (SidebandData.Num())
 	{
-		TSharedPtr<ElectraDecodersUtil::MPEG::FAVCDecoderConfigurationRecord, ESPMode::ThreadSafe> NewConfigRecord;
-		NewConfigRecord = MakeShared<ElectraDecodersUtil::MPEG::FAVCDecoderConfigurationRecord, ESPMode::ThreadSafe>();
-		NewConfigRecord->SetRawData(SidebandData.GetData(), SidebandData.Num());
-		if (NewConfigRecord->Parse())
+		TSharedPtr<ElectraDecodersUtil::MPEG::H264::FAVCDecoderConfigurationRecord, ESPMode::ThreadSafe> NewConfigRecord;
+		NewConfigRecord = MakeShared<ElectraDecodersUtil::MPEG::H264::FAVCDecoderConfigurationRecord, ESPMode::ThreadSafe>();
+		if (NewConfigRecord->Parse(SidebandData))
 		{
 			return NewConfigRecord;
 		}
@@ -340,19 +339,19 @@ TSharedPtr<ElectraDecodersUtil::MPEG::FAVCDecoderConfigurationRecord, ESPMode::T
 	return nullptr;
 }
 
-TSharedPtr<ElectraDecodersUtil::MPEG::FISO14496_10_seq_parameter_set_data, ESPMode::ThreadSafe> FElectraVideoDecoderH264_Android::GetSPSFromOptions(const TMap<FString, FVariant>& InOptions)
+TSharedPtr<ElectraDecodersUtil::MPEG::H264::FSequenceParameterSet, ESPMode::ThreadSafe> FElectraVideoDecoderH264_Android::GetSPSFromOptions(const TMap<FString, FVariant>& InOptions)
 {
 	TArray<uint8> SidebandData = ElectraDecodersUtil::GetVariantValueUInt8Array(InOptions, TEXT("csd"));
 	if (SidebandData.Num())
 	{
-		TArray<ElectraDecodersUtil::MPEG::FNaluInfo> NALUs;
-		ElectraDecodersUtil::MPEG::ParseBitstreamForNALUs(NALUs, SidebandData.GetData(), SidebandData.Num());
+		TArray<ElectraDecodersUtil::MPEG::H264::FNaluInfo> NALUs;
+		ElectraDecodersUtil::MPEG::H264::ParseBitstreamForNALUs(NALUs, SidebandData.GetData(), SidebandData.Num());
 		for(int32 i=0; i<NALUs.Num(); ++i)
 		{
-			if ((NALUs[i].Type & 0x1f) == 7)
+			if (NALUs[i].Type == 7)
 			{
-				TSharedPtr<ElectraDecodersUtil::MPEG::FISO14496_10_seq_parameter_set_data, ESPMode::ThreadSafe> NewSPS = MakeShared<ElectraDecodersUtil::MPEG::FISO14496_10_seq_parameter_set_data, ESPMode::ThreadSafe>();
-				if (ElectraDecodersUtil::MPEG::ParseH264SPS(*NewSPS, ElectraDecodersUtil::AdvancePointer(SidebandData.GetData(), NALUs[i].Offset + NALUs[i].UnitLength), NALUs[i].Size))
+				TSharedPtr<ElectraDecodersUtil::MPEG::H264::FSequenceParameterSet, ESPMode::ThreadSafe> NewSPS = MakeShared<ElectraDecodersUtil::MPEG::H264::FSequenceParameterSet, ESPMode::ThreadSafe>();
+				if (ElectraDecodersUtil::MPEG::H264::ParseSequenceParameterSet(*NewSPS, ElectraDecodersUtil::AdvancePointer(SidebandData.GetData(), NALUs[i].Offset + NALUs[i].UnitLength), NALUs[i].Size))
 				{
 					return NewSPS;
 				}
@@ -398,7 +397,7 @@ IElectraDecoder::ECSDCompatibility FElectraVideoDecoderH264_Android::IsCompatibl
 		return IElectraDecoder::ECSDCompatibility::Compatible;
 	}
 
-	TSharedPtr<ElectraDecodersUtil::MPEG::FISO14496_10_seq_parameter_set_data, ESPMode::ThreadSafe> NewSPS = GetSPSFromOptions(CSDAndAdditionalOptions);
+	TSharedPtr<ElectraDecodersUtil::MPEG::H264::FSequenceParameterSet, ESPMode::ThreadSafe> NewSPS = GetSPSFromOptions(CSDAndAdditionalOptions);
 	if (!NewSPS.IsValid())
 	{
 		UE_LOG(LogElectraDecoders, Warning, TEXT("No CSD provided to IsCompatibleWith(), returning to drain and reset."));
@@ -907,14 +906,17 @@ bool FElectraVideoDecoderH264_Android::InternalFinishDecoderCreate()
 		int32 LevelIdc = InitialMaxValues.Level ? InitialMaxValues.Level : 30;
 		if (ConfigRecord.IsValid())
 		{
-			check(ConfigRecord->GetNumberOfSPS() > 0);
-			const ElectraDecodersUtil::MPEG::FISO14496_10_seq_parameter_set_data& SPS = ConfigRecord->GetParsedSPS(0);
-			if (SPS.GetWidth() > Width || SPS.GetHeight() > Height || SPS.profile_idc > ProfileIdc || SPS.level_idc > LevelIdc)
+			if (ensure(!ConfigRecord->GetSequenceParameterSets().IsEmpty()))
 			{
-				Width = SPS.GetWidth();
-				Height = SPS.GetHeight();
-				ProfileIdc = SPS.profile_idc;
-				LevelIdc = SPS.level_idc;
+				ElectraDecodersUtil::MPEG::H264::FSequenceParameterSet SPS;
+				ElectraDecodersUtil::MPEG::H264::ParseSequenceParameterSet(SPS, ConfigRecord->GetSequenceParameterSets()[0]);
+				if (SPS.GetWidth() > Width || SPS.GetHeight() > Height || SPS.profile_idc > ProfileIdc || SPS.level_idc > LevelIdc)
+				{
+					Width = SPS.GetWidth();
+					Height = SPS.GetHeight();
+					ProfileIdc = SPS.profile_idc;
+					LevelIdc = SPS.level_idc;
+				}
 			}
 		}
 
diff --git a/Engine/Plugins/Media/ElectraCodecs/Source/ElectraDecoders/Private/Android/h264/VideoDecoderH264_JavaWrapper_Android.cpp b/Engine/Plugins/Media/ElectraCodecs/Source/ElectraDecoders/Private/Android/h264/VideoDecoderH264_JavaWrapper_Android.cpp
--- a/Engine/Plugins/Media/ElectraCodecs/Source/ElectraDecoders/Private/Android/h264/VideoDecoderH264_JavaWrapper_Android.cpp
+++ b/Engine/Plugins/Media/ElectraCodecs/Source/ElectraDecoders/Private/Android/h264/VideoDecoderH264_JavaWrapper_Android.cpp
@@ -4,7 +4,7 @@
 
 #if PLATFORM_ANDROID
 #include "HAL/Platform.h"
-#include "Utils/MPEG/ElectraUtilsMPEGVideo.h"
+#include "Utils/MPEG/ElectraUtilsMPEGVideo_H264.h"
 
 #if (defined(USE_ANDROID_JNI_WITHOUT_GAMEACTIVITY) && USE_ANDROID_JNI_WITHOUT_GAMEACTIVITY != 0) || USE_ANDROID_JNI
 #include "Android/AndroidPlatform.h"
@@ -452,13 +452,16 @@ int32 FElectraH264VideoDecoderAndroidJava::InitializeDecoder(const FCreateParame
 		//JEnv->SetBooleanField(CreateParams, FCreateParameters_bNeedTunneling, false);
 		JEnv->SetIntField(CreateParams, FCreateParameters_NativeDecoderID, InCreateParams.NativeDecoderID);
 
-		if (InCreateParams.CodecData && InCreateParams.CodecData->GetNumberOfSPS() > 0)
+		ElectraDecodersUtil::MPEG::H264::FSequenceParameterSet SPS;
+		if (InCreateParams.CodecData && 
+		   !InCreateParams.CodecData->GetSequenceParameterSets().IsEmpty() &&
+		   !InCreateParams.CodecData->GetPictureParameterSets().IsEmpty() &&
+		   ElectraDecodersUtil::MPEG::H264::ParseSequenceParameterSet(SPS, InCreateParams.CodecData->GetSequenceParameterSets()[0]))
 		{
-			const ElectraDecodersUtil::MPEG::FISO14496_10_seq_parameter_set_data& SPS = InCreateParams.CodecData->GetParsedSPS(0);
 			JEnv->SetIntField(CreateParams, FCreateParameters_Width, SPS.GetWidth());
 			JEnv->SetIntField(CreateParams, FCreateParameters_Height, SPS.GetHeight());
-			jbyteArray CSD0 = MakeJavaByteArray(InCreateParams.CodecData->GetCodecSpecificDataSPS().GetData(), InCreateParams.CodecData->GetCodecSpecificDataSPS().Num());
-			jbyteArray CSD1 = MakeJavaByteArray(InCreateParams.CodecData->GetCodecSpecificDataPPS().GetData(), InCreateParams.CodecData->GetCodecSpecificDataPPS().Num());
+			jbyteArray CSD0 = MakeJavaByteArray(InCreateParams.CodecData->GetSequenceParameterSets()[0].GetData(), InCreateParams.CodecData->GetSequenceParameterSets()[0].Num());
+			jbyteArray CSD1 = MakeJavaByteArray(InCreateParams.CodecData->GetPictureParameterSets()[0].GetData(), InCreateParams.CodecData->GetPictureParameterSets()[0].Num());
 			JEnv->SetObjectField(CreateParams, FCreateParameters_CSD0, CSD0);
 			JEnv->SetObjectField(CreateParams, FCreateParameters_CSD1, CSD1);
 			JEnv->DeleteLocalRef(CSD0);
diff --git a/Engine/Plugins/Media/ElectraCodecs/Source/ElectraDecoders/Private/Android/h264/VideoDecoderH264_JavaWrapper_Android.h b/Engine/Plugins/Media/ElectraCodecs/Source/ElectraDecoders/Private/Android/h264/VideoDecoderH264_JavaWrapper_Android.h
--- a/Engine/Plugins/Media/ElectraCodecs/Source/ElectraDecoders/Private/Android/h264/VideoDecoderH264_JavaWrapper_Android.h
+++ b/Engine/Plugins/Media/ElectraCodecs/Source/ElectraDecoders/Private/Android/h264/VideoDecoderH264_JavaWrapper_Android.h
@@ -10,7 +10,7 @@
 #include "Android/AndroidPlatform.h"
 #include "Android/AndroidJava.h"
 
-#include "Utils/MPEG/ElectraUtilsMPEGVideo.h"
+#include "Utils/MPEG/ElectraUtilsMPEGVideo_H264.h"
 
 class IElectraH264VideoDecoderAndroidJava
 {
@@ -27,7 +27,7 @@ public:
 		int32 MaxProfile = 0;
 		int32 MaxProfileLevel = 0;
 		int32 MaxFrameRate= 0;
-		const ElectraDecodersUtil::MPEG::FAVCDecoderConfigurationRecord* CodecData = nullptr;
+		const ElectraDecodersUtil::MPEG::H264::FAVCDecoderConfigurationRecord* CodecData = nullptr;
 		uint32 NativeDecoderID = 0;
 		bool bUseVideoCodecSurface = false;
 		jobject VideoCodecSurface = nullptr; // global reference expected
diff --git a/Engine/Plugins/Media/ElectraCodecs/Source/ElectraDecoders/Private/Android/h265/VideoDecoderH265_Android.cpp b/Engine/Plugins/Media/ElectraCodecs/Source/ElectraDecoders/Private/Android/h265/VideoDecoderH265_Android.cpp
--- a/Engine/Plugins/Media/ElectraCodecs/Source/ElectraDecoders/Private/Android/h265/VideoDecoderH265_Android.cpp
+++ b/Engine/Plugins/Media/ElectraCodecs/Source/ElectraDecoders/Private/Android/h265/VideoDecoderH265_Android.cpp
@@ -3,7 +3,7 @@
 #include "h265/VideoDecoderH265_Android.h"
 #include "DecoderErrors_Android.h"
 #include "ElectraDecodersUtils.h"
-#include "Utils/MPEG/ElectraUtilsMPEGVideo.h"
+#include "Utils/MPEG/ElectraUtilsMPEGVideo_H265.h"
 #include "Utils/MPEG/ElectraBitstreamProcessor_H265.h"
 
 #include "IElectraDecoderFeaturesAndOptions.h"
@@ -202,7 +202,7 @@ private:
 	{
 		FInputAccessUnit AccessUnit;
 		TMap<FString, FVariant> AdditionalOptions;
-		TSharedPtr<ElectraDecodersUtil::MPEG::FISO23008_2_seq_parameter_set_data, ESPMode::ThreadSafe> SPS;
+		TSharedPtr<ElectraDecodersUtil::MPEG::H265::FSequenceParameterSet, ESPMode::ThreadSafe> SPS;
 	};
 
 	enum class EDecodeState
@@ -214,8 +214,8 @@ private:
 
 	bool PostError(int32 ApiReturnValue, FString Message, int32 Code);
 
-	TSharedPtr<ElectraDecodersUtil::MPEG::FHEVCDecoderConfigurationRecord, ESPMode::ThreadSafe> GetDCRFromOptions(const TMap<FString, FVariant>& InOptions);
-	TSharedPtr<ElectraDecodersUtil::MPEG::FISO23008_2_seq_parameter_set_data, ESPMode::ThreadSafe> GetSPSFromOptions(const TMap<FString, FVariant>& InOptions);
+	TSharedPtr<ElectraDecodersUtil::MPEG::H265::FHEVCDecoderConfigurationRecord, ESPMode::ThreadSafe> GetDCRFromOptions(const TMap<FString, FVariant>& InOptions);
+	TSharedPtr<ElectraDecodersUtil::MPEG::H265::FSequenceParameterSet, ESPMode::ThreadSafe> GetSPSFromOptions(const TMap<FString, FVariant>& InOptions);
 
 	bool InternalDecoderCreate();
 	bool InternalHandleDecoderCreate();
@@ -239,8 +239,8 @@ private:
 	TWeakPtr<IElectraDecoderResourceDelegate, ESPMode::ThreadSafe> ResourceDelegate;
 	IElectraDecoderResourceDelegateAndroid::IDecoderPlatformResourceAndroid* PlatformResource = nullptr;
 
-	TSharedPtr<ElectraDecodersUtil::MPEG::FHEVCDecoderConfigurationRecord, ESPMode::ThreadSafe> ConfigRecord;
-	TSharedPtr<ElectraDecodersUtil::MPEG::FISO23008_2_seq_parameter_set_data, ESPMode::ThreadSafe> CurrentSPS;
+	TSharedPtr<ElectraDecodersUtil::MPEG::H265::FHEVCDecoderConfigurationRecord, ESPMode::ThreadSafe> ConfigRecord;
+	TSharedPtr<ElectraDecodersUtil::MPEG::H265::FSequenceParameterSet, ESPMode::ThreadSafe> CurrentSPS;
 	TSharedPtr<FElectraVideoDecoderOutputH265_Android, ESPMode::ThreadSafe> CurrentOutput;
 	TArray<TSharedPtr<FDecoderInput, ESPMode::NotThreadSafe>> InDecoderInput;
 
@@ -346,15 +346,14 @@ FElectraVideoDecoderH265_Android::~FElectraVideoDecoderH265_Android()
 }
 
 
-TSharedPtr<ElectraDecodersUtil::MPEG::FHEVCDecoderConfigurationRecord, ESPMode::ThreadSafe> FElectraVideoDecoderH265_Android::GetDCRFromOptions(const TMap<FString, FVariant>& InOptions)
+TSharedPtr<ElectraDecodersUtil::MPEG::H265::FHEVCDecoderConfigurationRecord, ESPMode::ThreadSafe> FElectraVideoDecoderH265_Android::GetDCRFromOptions(const TMap<FString, FVariant>& InOptions)
 {
 	TArray<uint8> SidebandData = ElectraDecodersUtil::GetVariantValueUInt8Array(InOptions, TEXT("dcr"));
 	if (SidebandData.Num())
 	{
-		TSharedPtr<ElectraDecodersUtil::MPEG::FHEVCDecoderConfigurationRecord, ESPMode::ThreadSafe> NewConfigRecord;
-		NewConfigRecord = MakeShared<ElectraDecodersUtil::MPEG::FHEVCDecoderConfigurationRecord, ESPMode::ThreadSafe>();
-		NewConfigRecord->SetRawData(SidebandData.GetData(), SidebandData.Num());
-		if (NewConfigRecord->Parse())
+		TSharedPtr<ElectraDecodersUtil::MPEG::H265::FHEVCDecoderConfigurationRecord, ESPMode::ThreadSafe> NewConfigRecord;
+		NewConfigRecord = MakeShared<ElectraDecodersUtil::MPEG::H265::FHEVCDecoderConfigurationRecord, ESPMode::ThreadSafe>();
+		if (NewConfigRecord->Parse(SidebandData))
 		{
 			return NewConfigRecord;
 		}
@@ -367,20 +366,19 @@ TSharedPtr<ElectraDecodersUtil::MPEG::FHEVCDecoderConfigurationRecord, ESPMode::
 	return nullptr;
 }
 
-TSharedPtr<ElectraDecodersUtil::MPEG::FISO23008_2_seq_parameter_set_data, ESPMode::ThreadSafe> FElectraVideoDecoderH265_Android::GetSPSFromOptions(const TMap<FString, FVariant>& InOptions)
+TSharedPtr<ElectraDecodersUtil::MPEG::H265::FSequenceParameterSet, ESPMode::ThreadSafe> FElectraVideoDecoderH265_Android::GetSPSFromOptions(const TMap<FString, FVariant>& InOptions)
 {
 	TArray<uint8> SidebandData = ElectraDecodersUtil::GetVariantValueUInt8Array(InOptions, TEXT("csd"));
 	if (SidebandData.Num())
 	{
-		TArray<ElectraDecodersUtil::MPEG::FNaluInfo> NALUs;
-		ElectraDecodersUtil::MPEG::ParseBitstreamForNALUs(NALUs, SidebandData.GetData(), SidebandData.Num());
+		TArray<ElectraDecodersUtil::MPEG::H265::FNaluInfo> NALUs;
+		ElectraDecodersUtil::MPEG::H265::ParseBitstreamForNALUs(NALUs, SidebandData.GetData(), SidebandData.Num());
 		for(int32 i=0; i<NALUs.Num(); ++i)
 		{
-			uint8 nut = NALUs[i].Type >> 1;
-			if (nut == 33)
+			if (NALUs[i].Type == 33)
 			{
-				TSharedPtr<ElectraDecodersUtil::MPEG::FISO23008_2_seq_parameter_set_data, ESPMode::ThreadSafe> NewSPS = MakeShared<ElectraDecodersUtil::MPEG::FISO23008_2_seq_parameter_set_data, ESPMode::ThreadSafe>();
-				if (ElectraDecodersUtil::MPEG::ParseH265SPS(*NewSPS, ElectraDecodersUtil::AdvancePointer(SidebandData.GetData(), NALUs[i].Offset + NALUs[i].UnitLength), NALUs[i].Size))
+				TSharedPtr<ElectraDecodersUtil::MPEG::H265::FSequenceParameterSet, ESPMode::ThreadSafe> NewSPS = MakeShared<ElectraDecodersUtil::MPEG::H265::FSequenceParameterSet, ESPMode::ThreadSafe>();
+				if (ElectraDecodersUtil::MPEG::H265::ParseSequenceParameterSet(*NewSPS, ElectraDecodersUtil::AdvancePointer(SidebandData.GetData(), NALUs[i].Offset + NALUs[i].UnitLength), NALUs[i].Size))
 				{
 					return NewSPS;
 				}
@@ -425,7 +423,7 @@ IElectraDecoder::ECSDCompatibility FElectraVideoDecoderH265_Android::IsCompatibl
 		return IElectraDecoder::ECSDCompatibility::Compatible;
 	}
 
-	TSharedPtr<ElectraDecodersUtil::MPEG::FISO23008_2_seq_parameter_set_data, ESPMode::ThreadSafe> NewSPS = GetSPSFromOptions(CSDAndAdditionalOptions);
+	TSharedPtr<ElectraDecodersUtil::MPEG::H265::FSequenceParameterSet, ESPMode::ThreadSafe> NewSPS = GetSPSFromOptions(CSDAndAdditionalOptions);
 	if (!NewSPS.IsValid())
 	{
 		UE_LOG(LogElectraDecoders, Warning, TEXT("No CSD provided to IsCompatibleWith(), returning to drain and reset."));
@@ -934,14 +932,16 @@ bool FElectraVideoDecoderH265_Android::InternalFinishDecoderCreate()
 		int32 LevelIdc = InitialMaxValues.Level ? InitialMaxValues.Level : 30;
 		if (ConfigRecord.IsValid())
 		{
-			check(ConfigRecord->GetNumberOfSPS() > 0);
-			const ElectraDecodersUtil::MPEG::FISO23008_2_seq_parameter_set_data& SPS = ConfigRecord->GetParsedSPS(0);
-			if (SPS.GetWidth() > Width || SPS.GetHeight() > Height || SPS.general_profile_idc > ProfileIdc || SPS.general_level_idc > LevelIdc)
+			if (ensure(!ConfigRecord->GetParsedSequenceParameterSets().IsEmpty()))
 			{
-				Width = SPS.GetWidth();
-				Height = SPS.GetHeight();
-				ProfileIdc = SPS.general_profile_idc;
-				LevelIdc = SPS.general_level_idc;
+				const ElectraDecodersUtil::MPEG::H265::FSequenceParameterSet& SPS = ConfigRecord->GetParsedSequenceParameterSets()[0];
+				if (SPS.GetWidth() > Width || SPS.GetHeight() > Height || SPS.profile_tier_level.general_profile_idc > ProfileIdc || SPS.profile_tier_level.general_level_idc > LevelIdc)
+				{
+					Width = SPS.GetWidth();
+					Height = SPS.GetHeight();
+					ProfileIdc = SPS.profile_tier_level.general_profile_idc;
+					LevelIdc = SPS.profile_tier_level.general_level_idc;
+				}
 			}
 		}
 
@@ -1110,10 +1110,11 @@ FElectraVideoDecoderH265_Android::EConvertResult FElectraVideoDecoderH265_Androi
 		NewOutput->AspectW = 1;
 		NewOutput->AspectH = 1;
 	}
-	if (In->SPS->vui_parameters_present_flag && In->SPS->vui_timing_info_present_flag && In->SPS->vui_time_scale)
+	ElectraDecodersUtil::FFractionalValue fr(In->SPS->GetTiming());
+	if (fr.Denom)
 	{
-		NewOutput->FrameRateN = In->SPS->vui_time_scale;
-		NewOutput->FrameRateD = In->SPS->vui_num_units_in_tick;
+		NewOutput->FrameRateN = fr.Num;
+		NewOutput->FrameRateD = fr.Denom;
 	}
 
 	NewOutput->OwningDecoder = StaticCastWeakPtr<FElectraVideoDecoderH265_Android>(AsWeak());
diff --git a/Engine/Plugins/Media/ElectraCodecs/Source/ElectraDecoders/Private/Android/h265/VideoDecoderH265_JavaWrapper_Android.cpp b/Engine/Plugins/Media/ElectraCodecs/Source/ElectraDecoders/Private/Android/h265/VideoDecoderH265_JavaWrapper_Android.cpp
--- a/Engine/Plugins/Media/ElectraCodecs/Source/ElectraDecoders/Private/Android/h265/VideoDecoderH265_JavaWrapper_Android.cpp
+++ b/Engine/Plugins/Media/ElectraCodecs/Source/ElectraDecoders/Private/Android/h265/VideoDecoderH265_JavaWrapper_Android.cpp
@@ -4,7 +4,7 @@
 
 #if PLATFORM_ANDROID
 #include "HAL/Platform.h"
-#include "Utils/MPEG/ElectraUtilsMPEGVideo.h"
+#include "Utils/MPEG/ElectraUtilsMPEGVideo_h265.h"
 
 #if (defined(USE_ANDROID_JNI_WITHOUT_GAMEACTIVITY) && USE_ANDROID_JNI_WITHOUT_GAMEACTIVITY != 0) || USE_ANDROID_JNI
 #include "Android/AndroidPlatform.h"
@@ -450,10 +450,10 @@ int32 FElectraH265VideoDecoderAndroidJava::InitializeDecoder(const FCreateParame
 		//JEnv->SetBooleanField(CreateParams, FCreateParameters_bNeedTunneling, false);
 		JEnv->SetIntField(CreateParams, FCreateParameters_NativeDecoderID, InCreateParams.NativeDecoderID);
 
-		if (InCreateParams.CodecData && InCreateParams.CodecData->GetNumberOfSPS() > 0)
+		if (InCreateParams.CodecData && !InCreateParams.CodecData->GetParsedSequenceParameterSets().IsEmpty())
 		{
-			JEnv->SetIntField(CreateParams, FCreateParameters_Width, InCreateParams.CodecData->GetParsedSPS(0).GetWidth());
-			JEnv->SetIntField(CreateParams, FCreateParameters_Height, InCreateParams.CodecData->GetParsedSPS(0).GetHeight());
+			JEnv->SetIntField(CreateParams, FCreateParameters_Width, InCreateParams.CodecData->GetParsedSequenceParameterSets()[0].GetWidth());
+			JEnv->SetIntField(CreateParams, FCreateParameters_Height, InCreateParams.CodecData->GetParsedSequenceParameterSets()[0].GetHeight());
 			jbyteArray CSD0 = MakeJavaByteArray(InCreateParams.CodecData->GetCodecSpecificData().GetData(), InCreateParams.CodecData->GetCodecSpecificData().Num());
 			JEnv->SetObjectField(CreateParams, FCreateParameters_CSD0, CSD0);
 			JEnv->DeleteLocalRef(CSD0);
diff --git a/Engine/Plugins/Media/ElectraCodecs/Source/ElectraDecoders/Private/Android/h265/VideoDecoderH265_JavaWrapper_Android.h b/Engine/Plugins/Media/ElectraCodecs/Source/ElectraDecoders/Private/Android/h265/VideoDecoderH265_JavaWrapper_Android.h
--- a/Engine/Plugins/Media/ElectraCodecs/Source/ElectraDecoders/Private/Android/h265/VideoDecoderH265_JavaWrapper_Android.h
+++ b/Engine/Plugins/Media/ElectraCodecs/Source/ElectraDecoders/Private/Android/h265/VideoDecoderH265_JavaWrapper_Android.h
@@ -10,7 +10,7 @@
 #include "Android/AndroidPlatform.h"
 #include "Android/AndroidJava.h"
 
-#include "Utils/MPEG/ElectraUtilsMPEGVideo.h"
+#include "Utils/MPEG/ElectraUtilsMPEGVideo_H265.h"
 
 class IElectraH265VideoDecoderAndroidJava
 {
@@ -28,7 +28,7 @@ public:
 		int32 MaxProfile = 0;
 		int32 MaxProfileLevel = 0;
 		int32 MaxFrameRate= 0;
-		const ElectraDecodersUtil::MPEG::FHEVCDecoderConfigurationRecord* CodecData = nullptr;
+		const ElectraDecodersUtil::MPEG::H265::FHEVCDecoderConfigurationRecord* CodecData = nullptr;
 		uint32 NativeDecoderID = 0;
 		bool bUseVideoCodecSurface = false;
 		jobject VideoCodecSurface = nullptr; // global reference expected
diff --git a/Engine/Plugins/Media/ElectraCodecs/Source/ElectraDecoders/Private/Apple/h264/VideoDecoderH264_Apple.cpp b/Engine/Plugins/Media/ElectraCodecs/Source/ElectraDecoders/Private/Apple/h264/VideoDecoderH264_Apple.cpp
--- a/Engine/Plugins/Media/ElectraCodecs/Source/ElectraDecoders/Private/Apple/h264/VideoDecoderH264_Apple.cpp
+++ b/Engine/Plugins/Media/ElectraCodecs/Source/ElectraDecoders/Private/Apple/h264/VideoDecoderH264_Apple.cpp
@@ -6,7 +6,7 @@
 
 #include "DecoderErrors_Apple.h"
 #include "ElectraDecodersUtils.h"
-#include "Utils/MPEG/ElectraUtilsMPEGVideo.h"
+#include "Utils/MPEG/ElectraUtilsMPEGVideo_H264.h"
 #include "Utils/MPEG/ElectraBitstreamProcessor_H264.h"
 
 #include "IElectraDecoderFeaturesAndOptions.h"
@@ -241,7 +241,7 @@ private:
 	{
 		FInputAccessUnit AccessUnit;
 		TMap<FString, FVariant> AdditionalOptions;
-		TSharedPtr<ElectraDecodersUtil::MPEG::FISO14496_10_seq_parameter_set_data, ESPMode::ThreadSafe> SPS;
+		TSharedPtr<ElectraDecodersUtil::MPEG::H264::FSequenceParameterSet, ESPMode::ThreadSafe> SPS;
 		bool bDropOutput = false;
 	};
 
@@ -318,7 +318,7 @@ private:
 
 	bool PostError(int32_t ApiReturnValue, FString Message, int32 Code);
 
-	TSharedPtr<ElectraDecodersUtil::MPEG::FISO14496_10_seq_parameter_set_data, ESPMode::ThreadSafe> GetSPSFromOptions(const TMap<FString, FVariant>& InOptions);
+	TSharedPtr<ElectraDecodersUtil::MPEG::H264::FSequenceParameterSet, ESPMode::ThreadSafe> GetSPSFromOptions(const TMap<FString, FVariant>& InOptions);
 	bool CreateFormatDescriptionFromOptions(CMFormatDescriptionRef& OutFormatDescription, const TMap<FString, FVariant>& InOptions);
 
 	bool InternalDecoderCreate(const TMap<FString, FVariant>& InAdditionalOptions);
@@ -343,7 +343,7 @@ private:
 	TArray<TSharedPtr<FDecoderInput, ESPMode::ThreadSafe>> InDecoderInput;
 
 	TSharedPtr<FElectraVideoDecoderOutputH264_Apple, ESPMode::ThreadSafe> CurrentOutput;
-	TSharedPtr<ElectraDecodersUtil::MPEG::FISO14496_10_seq_parameter_set_data, ESPMode::ThreadSafe> CurrentSPS;
+	TSharedPtr<ElectraDecodersUtil::MPEG::H264::FSequenceParameterSet, ESPMode::ThreadSafe> CurrentSPS;
 
 	TSharedPtr<FDecoderHandle, ESPMode::ThreadSafe> Decoder;
 
@@ -420,19 +420,19 @@ FElectraVideoDecoderH264_Apple::~FElectraVideoDecoderH264_Apple()
 	Close();
 }
 
-TSharedPtr<ElectraDecodersUtil::MPEG::FISO14496_10_seq_parameter_set_data, ESPMode::ThreadSafe> FElectraVideoDecoderH264_Apple::GetSPSFromOptions(const TMap<FString, FVariant>& InOptions)
+TSharedPtr<ElectraDecodersUtil::MPEG::H264::FSequenceParameterSet, ESPMode::ThreadSafe> FElectraVideoDecoderH264_Apple::GetSPSFromOptions(const TMap<FString, FVariant>& InOptions)
 {
 	TArray<uint8> SidebandData = ElectraDecodersUtil::GetVariantValueUInt8Array(InOptions, TEXT("csd"));
 	if (SidebandData.Num())
 	{
-		TArray<ElectraDecodersUtil::MPEG::FNaluInfo> NALUs;
-		ElectraDecodersUtil::MPEG::ParseBitstreamForNALUs(NALUs, SidebandData.GetData(), SidebandData.Num());
+		TArray<ElectraDecodersUtil::MPEG::H264::FNaluInfo> NALUs;
+		ElectraDecodersUtil::MPEG::H264::ParseBitstreamForNALUs(NALUs, SidebandData.GetData(), SidebandData.Num());
 		for(int32 i=0; i<NALUs.Num(); ++i)
 		{
-			if ((NALUs[i].Type & 0x1f) == 7)
+			if (NALUs[i].Type == 7)
 			{
-				TSharedPtr<ElectraDecodersUtil::MPEG::FISO14496_10_seq_parameter_set_data, ESPMode::ThreadSafe> NewSPS = MakeShared<ElectraDecodersUtil::MPEG::FISO14496_10_seq_parameter_set_data, ESPMode::ThreadSafe>();
-				if (ElectraDecodersUtil::MPEG::ParseH264SPS(*NewSPS, ElectraDecodersUtil::AdvancePointer(SidebandData.GetData(), NALUs[i].Offset + NALUs[i].UnitLength), NALUs[i].Size))
+				TSharedPtr<ElectraDecodersUtil::MPEG::H264::FSequenceParameterSet, ESPMode::ThreadSafe> NewSPS = MakeShared<ElectraDecodersUtil::MPEG::H264::FSequenceParameterSet, ESPMode::ThreadSafe>();
+				if (ElectraDecodersUtil::MPEG::H264::ParseSequenceParameterSet(*NewSPS, ElectraDecodersUtil::AdvancePointer(SidebandData.GetData(), NALUs[i].Offset + NALUs[i].UnitLength), NALUs[i].Size))
 				{
 					return NewSPS;
 				}
@@ -453,8 +453,8 @@ bool FElectraVideoDecoderH264_Apple::CreateFormatDescriptionFromOptions(CMFormat
 	TArray<uint8> SidebandData = ElectraDecodersUtil::GetVariantValueUInt8Array(InOptions, TEXT("csd"));
 	if (SidebandData.Num())
 	{
-		TArray<ElectraDecodersUtil::MPEG::FNaluInfo> NALUs;
-		ElectraDecodersUtil::MPEG::ParseBitstreamForNALUs(NALUs, SidebandData.GetData(), SidebandData.Num());
+		TArray<ElectraDecodersUtil::MPEG::H264::FNaluInfo> NALUs;
+		ElectraDecodersUtil::MPEG::H264::ParseBitstreamForNALUs(NALUs, SidebandData.GetData(), SidebandData.Num());
 		int32 NumRecords = NALUs.Num();
 		if (NumRecords)
 		{
diff --git a/Engine/Plugins/Media/ElectraCodecs/Source/ElectraDecoders/Private/Apple/h265/VideoDecoderH265_Apple.cpp b/Engine/Plugins/Media/ElectraCodecs/Source/ElectraDecoders/Private/Apple/h265/VideoDecoderH265_Apple.cpp
--- a/Engine/Plugins/Media/ElectraCodecs/Source/ElectraDecoders/Private/Apple/h265/VideoDecoderH265_Apple.cpp
+++ b/Engine/Plugins/Media/ElectraCodecs/Source/ElectraDecoders/Private/Apple/h265/VideoDecoderH265_Apple.cpp
@@ -6,7 +6,7 @@
 
 #include "DecoderErrors_Apple.h"
 #include "ElectraDecodersUtils.h"
-#include "Utils/MPEG/ElectraUtilsMPEGVideo.h"
+#include "Utils/MPEG/ElectraUtilsMPEGVideo_H265.h"
 #include "Utils/MPEG/ElectraBitstreamProcessor_H265.h"
 
 #include "IElectraDecoderFeaturesAndOptions.h"
@@ -241,7 +241,7 @@ private:
 	{
 		FInputAccessUnit AccessUnit;
 		TMap<FString, FVariant> AdditionalOptions;
-		TSharedPtr<ElectraDecodersUtil::MPEG::FISO23008_2_seq_parameter_set_data, ESPMode::ThreadSafe> SPS;
+		TSharedPtr<ElectraDecodersUtil::MPEG::H265::FSequenceParameterSet, ESPMode::ThreadSafe> SPS;
 		bool bDropOutput = false;
 	};
 
@@ -318,7 +318,7 @@ private:
 
 	bool PostError(int32_t ApiReturnValue, FString Message, int32 Code);
 
-	TSharedPtr<ElectraDecodersUtil::MPEG::FISO23008_2_seq_parameter_set_data, ESPMode::ThreadSafe> GetSPSFromOptions(const TMap<FString, FVariant>& InOptions);
+	TSharedPtr<ElectraDecodersUtil::MPEG::H265::FSequenceParameterSet, ESPMode::ThreadSafe> GetSPSFromOptions(const TMap<FString, FVariant>& InOptions);
 	bool CreateFormatDescriptionFromOptions(CMFormatDescriptionRef& OutFormatDescription, const TMap<FString, FVariant>& InOptions);
 
 	bool InternalDecoderCreate(const TMap<FString, FVariant>& InAdditionalOptions);
@@ -343,7 +343,7 @@ private:
 	TArray<TSharedPtr<FDecoderInput, ESPMode::ThreadSafe>> InDecoderInput;
 
 	TSharedPtr<FElectraVideoDecoderOutputH265_Apple, ESPMode::ThreadSafe> CurrentOutput;
-	TSharedPtr<ElectraDecodersUtil::MPEG::FISO23008_2_seq_parameter_set_data, ESPMode::ThreadSafe> CurrentSPS;
+	TSharedPtr<ElectraDecodersUtil::MPEG::H265::FSequenceParameterSet, ESPMode::ThreadSafe> CurrentSPS;
 
 	TSharedPtr<FDecoderHandle, ESPMode::ThreadSafe> Decoder;
 
@@ -369,7 +369,7 @@ void IElectraVideoDecoderH265_Apple::PlatformGetSupportedConfigurations(TArray<F
 		IElectraVideoDecoderH265_Apple_Platform::DecoderConfigurations.Empty();
 
 		// Perhaps these can be determined dynamically for each type of device (Mac, iPad, iPhone, AppleTV)?
-		// What is really supported isn't quite clear so allow for UHD.
+		// What is really supported isn't quite clear so allow for 4K.
 
 		// Main
 		IElectraVideoDecoderH265_Apple_Platform::DecoderConfigurations.Emplace(IElectraVideoDecoderH265_Apple::FSupportedConfiguration(0, 1, 0, 153, 120, 4096, 2304, 0));
@@ -420,20 +420,19 @@ FElectraVideoDecoderH265_Apple::~FElectraVideoDecoderH265_Apple()
 	Close();
 }
 
-TSharedPtr<ElectraDecodersUtil::MPEG::FISO23008_2_seq_parameter_set_data, ESPMode::ThreadSafe> FElectraVideoDecoderH265_Apple::GetSPSFromOptions(const TMap<FString, FVariant>& InOptions)
+TSharedPtr<ElectraDecodersUtil::MPEG::H265::FSequenceParameterSet, ESPMode::ThreadSafe> FElectraVideoDecoderH265_Apple::GetSPSFromOptions(const TMap<FString, FVariant>& InOptions)
 {
 	TArray<uint8> SidebandData = ElectraDecodersUtil::GetVariantValueUInt8Array(InOptions, TEXT("csd"));
 	if (SidebandData.Num())
 	{
-		TArray<ElectraDecodersUtil::MPEG::FNaluInfo> NALUs;
-		ElectraDecodersUtil::MPEG::ParseBitstreamForNALUs(NALUs, SidebandData.GetData(), SidebandData.Num());
+		TArray<ElectraDecodersUtil::MPEG::H265::FNaluInfo> NALUs;
+		ElectraDecodersUtil::MPEG::H265::ParseBitstreamForNALUs(NALUs, SidebandData.GetData(), SidebandData.Num());
 		for(int32 i=0; i<NALUs.Num(); ++i)
 		{
-			uint8 nut = NALUs[i].Type >> 1;
-			if (nut == 33)
+			if (NALUs[i].Type == 33)
 			{
-				TSharedPtr<ElectraDecodersUtil::MPEG::FISO23008_2_seq_parameter_set_data, ESPMode::ThreadSafe> NewSPS = MakeShared<ElectraDecodersUtil::MPEG::FISO23008_2_seq_parameter_set_data, ESPMode::ThreadSafe>();
-				if (ElectraDecodersUtil::MPEG::ParseH265SPS(*NewSPS, ElectraDecodersUtil::AdvancePointer(SidebandData.GetData(), NALUs[i].Offset + NALUs[i].UnitLength), NALUs[i].Size))
+				TSharedPtr<ElectraDecodersUtil::MPEG::H265::FSequenceParameterSet, ESPMode::ThreadSafe> NewSPS = MakeShared<ElectraDecodersUtil::MPEG::H265::FSequenceParameterSet, ESPMode::ThreadSafe>();
+				if (ElectraDecodersUtil::MPEG::H265::ParseSequenceParameterSet(*NewSPS, ElectraDecodersUtil::AdvancePointer(SidebandData.GetData(), NALUs[i].Offset + NALUs[i].UnitLength), NALUs[i].Size))
 				{
 					return NewSPS;
 				}
@@ -456,8 +455,8 @@ bool FElectraVideoDecoderH265_Apple::CreateFormatDescriptionFromOptions(CMFormat
 	TArray<uint8> SidebandData = ElectraDecodersUtil::GetVariantValueUInt8Array(InOptions, TEXT("csd"));
 	if (SidebandData.Num())
 	{
-		TArray<ElectraDecodersUtil::MPEG::FNaluInfo> NALUs;
-		ElectraDecodersUtil::MPEG::ParseBitstreamForNALUs(NALUs, SidebandData.GetData(), SidebandData.Num());
+		TArray<ElectraDecodersUtil::MPEG::H265::FNaluInfo> NALUs;
+		ElectraDecodersUtil::MPEG::H265::ParseBitstreamForNALUs(NALUs, SidebandData.GetData(), SidebandData.Num());
 		int32 NumRecords = NALUs.Num();
 		if (NumRecords)
 		{
@@ -948,10 +947,12 @@ FElectraVideoDecoderH265_Apple::EConvertResult FElectraVideoDecoderH265_Apple::C
 		NewOutput->AspectW = 1;
 		NewOutput->AspectH = 1;
 	}
-	if (In->SPS->vui_parameters_present_flag && In->SPS->vui_timing_info_present_flag && In->SPS->vui_time_scale)
+
+	ElectraDecodersUtil::FFractionalValue fr(In->SPS->GetTiming());
+	if (fr.Denom)
 	{
-		NewOutput->FrameRateN = In->SPS->vui_time_scale;
-		NewOutput->FrameRateD = In->SPS->vui_num_units_in_tick;
+		NewOutput->FrameRateN = fr.Num;
+		NewOutput->FrameRateD = fr.Denom;
 	}
 
 	NewOutput->ExtraValues.Emplace(TEXT("platform"), FVariant(TEXT("apple")));
diff --git a/Engine/Plugins/Media/ElectraCodecs/Source/ElectraDecoders/Private/DX/VideoDecoderH265_DX.cpp b/Engine/Plugins/Media/ElectraCodecs/Source/ElectraDecoders/Private/DX/VideoDecoderH265_DX.cpp
--- a/Engine/Plugins/Media/ElectraCodecs/Source/ElectraDecoders/Private/DX/VideoDecoderH265_DX.cpp
+++ b/Engine/Plugins/Media/ElectraCodecs/Source/ElectraDecoders/Private/DX/VideoDecoderH265_DX.cpp
@@ -5,7 +5,7 @@
 #include "DX/VideoDecoderH265_DX.h"
 #include "DX/DecoderErrors_DX.h"
 #include "ElectraDecodersUtils.h"
-#include "Utils/MPEG/ElectraUtilsMPEGVideo.h"
+#include "Utils/MPEG/ElectraUtilsMPEGVideo_H265.h"
 #include "Utils/MPEG/ElectraBitstreamProcessor_H265.h"
 
 #include "IElectraDecoderFeaturesAndOptions.h"
@@ -347,7 +347,7 @@ private:
 
 	void PostError(HRESULT ApiReturnValue, FString Message, int32 Code);
 
-	TSharedPtr<ElectraDecodersUtil::MPEG::FISO23008_2_seq_parameter_set_data, ESPMode::ThreadSafe> GetSPSFromOptions(const TMap<FString, FVariant>& InOptions);
+	TSharedPtr<ElectraDecodersUtil::MPEG::H265::FSequenceParameterSet, ESPMode::ThreadSafe> GetSPSFromOptions(const TMap<FString, FVariant>& InOptions);
 
 	bool InternalDecoderCreate(const TMap<FString, FVariant>& InAdditionalOptions);
 	void InternalDecoderDestroy();
@@ -422,20 +422,19 @@ FElectraVideoDecoderH265_DX::~FElectraVideoDecoderH265_DX()
 	Close();
 }
 
-TSharedPtr<ElectraDecodersUtil::MPEG::FISO23008_2_seq_parameter_set_data, ESPMode::ThreadSafe> FElectraVideoDecoderH265_DX::GetSPSFromOptions(const TMap<FString, FVariant>& InOptions)
+TSharedPtr<ElectraDecodersUtil::MPEG::H265::FSequenceParameterSet, ESPMode::ThreadSafe> FElectraVideoDecoderH265_DX::GetSPSFromOptions(const TMap<FString, FVariant>& InOptions)
 {
 	TArray<uint8> SidebandData = ElectraDecodersUtil::GetVariantValueUInt8Array(InOptions, TEXT("csd"));
 	if (SidebandData.Num())
 	{
-		TArray<ElectraDecodersUtil::MPEG::FNaluInfo> NALUs;
-		ElectraDecodersUtil::MPEG::ParseBitstreamForNALUs(NALUs, SidebandData.GetData(), SidebandData.Num());
+		TArray<ElectraDecodersUtil::MPEG::H265::FNaluInfo> NALUs;
+		ElectraDecodersUtil::MPEG::H265::ParseBitstreamForNALUs(NALUs, SidebandData.GetData(), SidebandData.Num());
 		for(int32 i=0; i<NALUs.Num(); ++i)
 		{
-			uint8 nut = NALUs[i].Type >> 1;
-			if (nut == 33)
+			if (NALUs[i].Type == 33)
 			{
-				TSharedPtr<ElectraDecodersUtil::MPEG::FISO23008_2_seq_parameter_set_data, ESPMode::ThreadSafe> NewSPS = MakeShared<ElectraDecodersUtil::MPEG::FISO23008_2_seq_parameter_set_data, ESPMode::ThreadSafe>();
-				if (ElectraDecodersUtil::MPEG::ParseH265SPS(*NewSPS, ElectraDecodersUtil::AdvancePointer(SidebandData.GetData(), NALUs[i].Offset + NALUs[i].UnitLength), NALUs[i].Size))
+				TSharedPtr<ElectraDecodersUtil::MPEG::H265::FSequenceParameterSet, ESPMode::ThreadSafe> NewSPS = MakeShared<ElectraDecodersUtil::MPEG::H265::FSequenceParameterSet, ESPMode::ThreadSafe>();
+				if (ElectraDecodersUtil::MPEG::H265::ParseSequenceParameterSet(*NewSPS, ElectraDecodersUtil::AdvancePointer(SidebandData.GetData(), NALUs[i].Offset + NALUs[i].UnitLength), NALUs[i].Size))
 				{
 					return NewSPS;
 				}
@@ -819,11 +818,10 @@ bool FElectraVideoDecoderH265_DX::InternalDecoderCreate(const TMap<FString, FVar
 	VERIFY_HR(MediaType->SetBlob(MF_MT_USER_DATA, SidebandData.GetData(), SidebandData.Num()), TEXT("Failed to set input video CSD"), ERRCODE_INTERNAL_FAILED_TO_CREATE_DECODER);
 	// We need to set the resolution of the to-be-decoded stream, otherwise setting the output type will fail for some reason.
 	// If possible we get the actual resolution from the CSD, otherwise we set a common value. The exact size doesn't seem to be important.
-	ElectraDecodersUtil::MPEG::FHEVCDecoderConfigurationRecord dcr;
-	dcr.SetRawData(SidebandData.GetData(), SidebandData.Num());
-	if (dcr.Parse() && dcr.GetNumberOfSPS() > 0)
+	ElectraDecodersUtil::MPEG::H265::FHEVCDecoderConfigurationRecord dcr;
+	if (dcr.Parse(SidebandData) && !dcr.GetParsedSequenceParameterSets().IsEmpty())
 	{
-		VERIFY_HR(MFSetAttributeSize(MediaType, MF_MT_FRAME_SIZE, dcr.GetParsedSPS(0).GetWidth(), dcr.GetParsedSPS(0).GetHeight()), "Failed to set video decoder input media type resolution", ERRCODE_INTERNAL_FAILED_TO_CREATE_DECODER);
+		VERIFY_HR(MFSetAttributeSize(MediaType, MF_MT_FRAME_SIZE, dcr.GetParsedSequenceParameterSets()[0].GetWidth(), dcr.GetParsedSequenceParameterSets()[0].GetHeight()), "Failed to set video decoder input media type resolution", ERRCODE_INTERNAL_FAILED_TO_CREATE_DECODER);
 	}
 	else
 	{
diff --git a/Engine/Plugins/Media/ElectraCodecs/Source/ElectraDecoders/Private/Linux/h264/VideoDecoderH264_Linux.cpp b/Engine/Plugins/Media/ElectraCodecs/Source/ElectraDecoders/Private/Linux/h264/VideoDecoderH264_Linux.cpp
--- a/Engine/Plugins/Media/ElectraCodecs/Source/ElectraDecoders/Private/Linux/h264/VideoDecoderH264_Linux.cpp
+++ b/Engine/Plugins/Media/ElectraCodecs/Source/ElectraDecoders/Private/Linux/h264/VideoDecoderH264_Linux.cpp
@@ -3,7 +3,7 @@
 #include "h264/VideoDecoderH264_Linux.h"
 #include "DecoderErrors_Linux.h"
 #include "ElectraDecodersUtils.h"
-#include "Utils/MPEG/ElectraUtilsMPEGVideo.h"
+#include "Utils/MPEG/ElectraUtilsMPEGVideo_H264.h"
 #include "Utils/MPEG/ElectraBitstreamProcessor_H264.h"
 
 #include "IElectraDecoderFeaturesAndOptions.h"
@@ -144,7 +144,7 @@ private:
 	{
 		FInputAccessUnit AccessUnit;
 		TMap<FString, FVariant> AdditionalOptions;
-		TSharedPtr<ElectraDecodersUtil::MPEG::FISO14496_10_seq_parameter_set_data, ESPMode::ThreadSafe> SPS;
+		TSharedPtr<ElectraDecodersUtil::MPEG::H264::FSequenceParameterSet, ESPMode::ThreadSafe> SPS;
 		bool bDropOutput = false;
 	};
 
@@ -156,7 +156,7 @@ private:
 
 	bool PostError(int32_t ApiReturnValue, FString Message, int32 Code);
 
-	TSharedPtr<ElectraDecodersUtil::MPEG::FISO14496_10_seq_parameter_set_data, ESPMode::ThreadSafe> GetSPSFromOptions(const TMap<FString, FVariant>& InOptions);
+	TSharedPtr<ElectraDecodersUtil::MPEG::H264::FSequenceParameterSet, ESPMode::ThreadSafe> GetSPSFromOptions(const TMap<FString, FVariant>& InOptions);
 
 	bool InternalDecoderCreate(const TMap<FString, FVariant>& InAdditionalOptions);
 	void InternalDecoderDestroy();
@@ -175,7 +175,7 @@ private:
 
 	TArray<TSharedPtr<FDecoderInput, ESPMode::NotThreadSafe>> InDecoderInput;
 	TSharedPtr<FElectraVideoDecoderOutputH264_Linux, ESPMode::ThreadSafe> CurrentOutput;
-	TSharedPtr<ElectraDecodersUtil::MPEG::FISO14496_10_seq_parameter_set_data, ESPMode::ThreadSafe> CurrentSPS;
+	TSharedPtr<ElectraDecodersUtil::MPEG::H264::FSequenceParameterSet, ESPMode::ThreadSafe> CurrentSPS;
 
 	TSharedPtr<ILibavDecoderVideoCommon, ESPMode::ThreadSafe> DecoderInstance;
 	EDecodeState DecodeState = EDecodeState::Decoding;
@@ -241,19 +241,19 @@ FElectraVideoDecoderH264_Linux::~FElectraVideoDecoderH264_Linux()
 	Close();
 }
 
-TSharedPtr<ElectraDecodersUtil::MPEG::FISO14496_10_seq_parameter_set_data, ESPMode::ThreadSafe> FElectraVideoDecoderH264_Linux::GetSPSFromOptions(const TMap<FString, FVariant>& InOptions)
+TSharedPtr<ElectraDecodersUtil::MPEG::H264::FSequenceParameterSet, ESPMode::ThreadSafe> FElectraVideoDecoderH264_Linux::GetSPSFromOptions(const TMap<FString, FVariant>& InOptions)
 {
 	TArray<uint8> SidebandData = ElectraDecodersUtil::GetVariantValueUInt8Array(InOptions, TEXT("csd"));
 	if (SidebandData.Num())
 	{
-		TArray<ElectraDecodersUtil::MPEG::FNaluInfo> NALUs;
-		ElectraDecodersUtil::MPEG::ParseBitstreamForNALUs(NALUs, SidebandData.GetData(), SidebandData.Num());
+		TArray<ElectraDecodersUtil::MPEG::H264::FNaluInfo> NALUs;
+		ElectraDecodersUtil::MPEG::H264::ParseBitstreamForNALUs(NALUs, SidebandData.GetData(), SidebandData.Num());
 		for(int32 i=0; i<NALUs.Num(); ++i)
 		{
-			if ((NALUs[i].Type & 0x1f) == 7)
+			if (NALUs[i].Type == 7)
 			{
-				TSharedPtr<ElectraDecodersUtil::MPEG::FISO14496_10_seq_parameter_set_data, ESPMode::ThreadSafe> NewSPS = MakeShared<ElectraDecodersUtil::MPEG::FISO14496_10_seq_parameter_set_data, ESPMode::ThreadSafe>();
-				if (ElectraDecodersUtil::MPEG::ParseH264SPS(*NewSPS, ElectraDecodersUtil::AdvancePointer(SidebandData.GetData(), NALUs[i].Offset + NALUs[i].UnitLength), NALUs[i].Size))
+				TSharedPtr<ElectraDecodersUtil::MPEG::H264::FSequenceParameterSet, ESPMode::ThreadSafe> NewSPS = MakeShared<ElectraDecodersUtil::MPEG::H264::FSequenceParameterSet, ESPMode::ThreadSafe>();
+				if (ElectraDecodersUtil::MPEG::H264::ParseSequenceParameterSet(*NewSPS, ElectraDecodersUtil::AdvancePointer(SidebandData.GetData(), NALUs[i].Offset + NALUs[i].UnitLength), NALUs[i].Size))
 				{
 					return NewSPS;
 				}
@@ -313,7 +313,7 @@ IElectraDecoder::ECSDCompatibility FElectraVideoDecoderH264_Linux::IsCompatibleW
 		return IElectraDecoder::ECSDCompatibility::Compatible;
 	}
 
-	TSharedPtr<ElectraDecodersUtil::MPEG::FISO14496_10_seq_parameter_set_data, ESPMode::ThreadSafe> NewSPS = GetSPSFromOptions(CSDAndAdditionalOptions);
+	TSharedPtr<ElectraDecodersUtil::MPEG::H264::FSequenceParameterSet, ESPMode::ThreadSafe> NewSPS = GetSPSFromOptions(CSDAndAdditionalOptions);
 	if (!NewSPS.IsValid())
 	{
 		UE_LOG(LogElectraDecoders, Warning, TEXT("No CSD provided to IsCompatibleWith(), returning to drain and reset."));
diff --git a/Engine/Plugins/Media/ElectraCodecs/Source/ElectraDecoders/Private/Linux/h265/VideoDecoderH265_Linux.cpp b/Engine/Plugins/Media/ElectraCodecs/Source/ElectraDecoders/Private/Linux/h265/VideoDecoderH265_Linux.cpp
--- a/Engine/Plugins/Media/ElectraCodecs/Source/ElectraDecoders/Private/Linux/h265/VideoDecoderH265_Linux.cpp
+++ b/Engine/Plugins/Media/ElectraCodecs/Source/ElectraDecoders/Private/Linux/h265/VideoDecoderH265_Linux.cpp
@@ -3,7 +3,7 @@
 #include "h265/VideoDecoderH265_Linux.h"
 #include "DecoderErrors_Linux.h"
 #include "ElectraDecodersUtils.h"
-#include "Utils/MPEG/ElectraUtilsMPEGVideo.h"
+#include "Utils/MPEG/ElectraUtilsMPEGVideo_H265.h"
 #include "Utils/MPEG/ElectraBitstreamProcessor_H265.h"
 
 #include "IElectraDecoderFeaturesAndOptions.h"
@@ -144,7 +144,7 @@ private:
 	{
 		FInputAccessUnit AccessUnit;
 		TMap<FString, FVariant> AdditionalOptions;
-		TSharedPtr<ElectraDecodersUtil::MPEG::FISO23008_2_seq_parameter_set_data, ESPMode::ThreadSafe> SPS;
+		TSharedPtr<ElectraDecodersUtil::MPEG::H265::FSequenceParameterSet, ESPMode::ThreadSafe> SPS;
 		bool bDropOutput = false;
 	};
 
@@ -156,8 +156,8 @@ private:
 
 	bool PostError(int32_t ApiReturnValue, FString Message, int32 Code);
 
-	TSharedPtr<ElectraDecodersUtil::MPEG::FHEVCDecoderConfigurationRecord, ESPMode::ThreadSafe> GetDCRFromOptions(const TMap<FString, FVariant>& InOptions);
-	TSharedPtr<ElectraDecodersUtil::MPEG::FISO23008_2_seq_parameter_set_data, ESPMode::ThreadSafe> GetSPSFromOptions(const TMap<FString, FVariant>& InOptions);
+	TSharedPtr<ElectraDecodersUtil::MPEG::H265::FHEVCDecoderConfigurationRecord, ESPMode::ThreadSafe> GetDCRFromOptions(const TMap<FString, FVariant>& InOptions);
+	TSharedPtr<ElectraDecodersUtil::MPEG::H265::FSequenceParameterSet, ESPMode::ThreadSafe> GetSPSFromOptions(const TMap<FString, FVariant>& InOptions);
 
 	bool InternalDecoderCreate(const TMap<FString, FVariant>& InAdditionalOptions);
 	void InternalDecoderDestroy();
@@ -176,7 +176,7 @@ private:
 
 	TArray<TSharedPtr<FDecoderInput, ESPMode::NotThreadSafe>> InDecoderInput;
 	TSharedPtr<FElectraVideoDecoderOutputH265_Linux, ESPMode::ThreadSafe> CurrentOutput;
-	TSharedPtr<ElectraDecodersUtil::MPEG::FISO23008_2_seq_parameter_set_data, ESPMode::ThreadSafe> CurrentSPS;
+	TSharedPtr<ElectraDecodersUtil::MPEG::H265::FSequenceParameterSet, ESPMode::ThreadSafe> CurrentSPS;
 
 	TSharedPtr<ILibavDecoderVideoCommon, ESPMode::ThreadSafe> DecoderInstance;
 	EDecodeState DecodeState = EDecodeState::Decoding;
@@ -243,15 +243,14 @@ FElectraVideoDecoderH265_Linux::~FElectraVideoDecoderH265_Linux()
 	Close();
 }
 
-TSharedPtr<ElectraDecodersUtil::MPEG::FHEVCDecoderConfigurationRecord, ESPMode::ThreadSafe> FElectraVideoDecoderH265_Linux::GetDCRFromOptions(const TMap<FString, FVariant>& InOptions)
+TSharedPtr<ElectraDecodersUtil::MPEG::H265::FHEVCDecoderConfigurationRecord, ESPMode::ThreadSafe> FElectraVideoDecoderH265_Linux::GetDCRFromOptions(const TMap<FString, FVariant>& InOptions)
 {
 	TArray<uint8> SidebandData = ElectraDecodersUtil::GetVariantValueUInt8Array(InOptions, TEXT("dcr"));
 	if (SidebandData.Num())
 	{
-		TSharedPtr<ElectraDecodersUtil::MPEG::FHEVCDecoderConfigurationRecord, ESPMode::ThreadSafe> NewConfigRecord;
-		NewConfigRecord = MakeShared<ElectraDecodersUtil::MPEG::FHEVCDecoderConfigurationRecord, ESPMode::ThreadSafe>();
-		NewConfigRecord->SetRawData(SidebandData.GetData(), SidebandData.Num());
-		if (NewConfigRecord->Parse())
+		TSharedPtr<ElectraDecodersUtil::MPEG::H265::FHEVCDecoderConfigurationRecord, ESPMode::ThreadSafe> NewConfigRecord;
+		NewConfigRecord = MakeShared<ElectraDecodersUtil::MPEG::H265::FHEVCDecoderConfigurationRecord, ESPMode::ThreadSafe>();
+		if (NewConfigRecord->Parse(SidebandData))
 		{
 			return NewConfigRecord;
 		}
@@ -264,20 +263,19 @@ TSharedPtr<ElectraDecodersUtil::MPEG::FHEVCDecoderConfigurationRecord, ESPMode::
 	return nullptr;
 }
 
-TSharedPtr<ElectraDecodersUtil::MPEG::FISO23008_2_seq_parameter_set_data, ESPMode::ThreadSafe> FElectraVideoDecoderH265_Linux::GetSPSFromOptions(const TMap<FString, FVariant>& InOptions)
+TSharedPtr<ElectraDecodersUtil::MPEG::H265::FSequenceParameterSet, ESPMode::ThreadSafe> FElectraVideoDecoderH265_Linux::GetSPSFromOptions(const TMap<FString, FVariant>& InOptions)
 {
 	TArray<uint8> SidebandData = ElectraDecodersUtil::GetVariantValueUInt8Array(InOptions, TEXT("csd"));
 	if (SidebandData.Num())
 	{
-		TArray<ElectraDecodersUtil::MPEG::FNaluInfo> NALUs;
-		ElectraDecodersUtil::MPEG::ParseBitstreamForNALUs(NALUs, SidebandData.GetData(), SidebandData.Num());
+		TArray<ElectraDecodersUtil::MPEG::H265::FNaluInfo> NALUs;
+		ElectraDecodersUtil::MPEG::H265::ParseBitstreamForNALUs(NALUs, SidebandData.GetData(), SidebandData.Num());
 		for(int32 i=0; i<NALUs.Num(); ++i)
 		{
-			uint8 nut = NALUs[i].Type >> 1;
-			if (nut == 33)
+			if (NALUs[i].Type == 33)
 			{
-				TSharedPtr<ElectraDecodersUtil::MPEG::FISO23008_2_seq_parameter_set_data, ESPMode::ThreadSafe> NewSPS = MakeShared<ElectraDecodersUtil::MPEG::FISO23008_2_seq_parameter_set_data, ESPMode::ThreadSafe>();
-				if (ElectraDecodersUtil::MPEG::ParseH265SPS(*NewSPS, ElectraDecodersUtil::AdvancePointer(SidebandData.GetData(), NALUs[i].Offset + NALUs[i].UnitLength), NALUs[i].Size))
+				TSharedPtr<ElectraDecodersUtil::MPEG::H265::FSequenceParameterSet, ESPMode::ThreadSafe> NewSPS = MakeShared<ElectraDecodersUtil::MPEG::H265::FSequenceParameterSet, ESPMode::ThreadSafe>();
+				if (ElectraDecodersUtil::MPEG::H265::ParseSequenceParameterSet(*NewSPS, ElectraDecodersUtil::AdvancePointer(SidebandData.GetData(), NALUs[i].Offset + NALUs[i].UnitLength), NALUs[i].Size))
 				{
 					return NewSPS;
 				}
@@ -336,7 +334,7 @@ IElectraDecoder::ECSDCompatibility FElectraVideoDecoderH265_Linux::IsCompatibleW
 		return IElectraDecoder::ECSDCompatibility::Compatible;
 	}
 
-	TSharedPtr<ElectraDecodersUtil::MPEG::FISO23008_2_seq_parameter_set_data, ESPMode::ThreadSafe> NewSPS = GetSPSFromOptions(CSDAndAdditionalOptions);
+	TSharedPtr<ElectraDecodersUtil::MPEG::H265::FSequenceParameterSet, ESPMode::ThreadSafe> NewSPS = GetSPSFromOptions(CSDAndAdditionalOptions);
 	if (!NewSPS.IsValid())
 	{
 		UE_LOG(LogElectraDecoders, Warning, TEXT("No CSD provided to IsCompatibleWith(), returning to drain and reset."));
@@ -637,10 +635,11 @@ FElectraVideoDecoderH265_Linux::EConvertResult FElectraVideoDecoderH265_Linux::C
 		NewOutput->AspectW = 1;
 		NewOutput->AspectH = 1;
 	}
-	if (In->SPS->vui_parameters_present_flag && In->SPS->vui_timing_info_present_flag && In->SPS->vui_time_scale)
+	ElectraDecodersUtil::FFractionalValue fr(In->SPS->GetTiming());
+	if (fr.Denom)
 	{
-		NewOutput->FrameRateN = In->SPS->vui_time_scale;
-		NewOutput->FrameRateD = In->SPS->vui_num_units_in_tick;
+		NewOutput->FrameRateN = fr.Num;
+		NewOutput->FrameRateD = fr.Denom;
 	}
 
 	NewOutput->ExtraValues.Emplace(TEXT("platform"), FVariant(TEXT("linux")));
diff --git a/Engine/Plugins/Media/ElectraCodecs/Source/ElectraDecoders/Private/Utils/MPEG/ElectraUtilsMPEGVideo.cpp b/Engine/Plugins/Media/ElectraCodecs/Source/ElectraDecoders/Private/Utils/MPEG/ElectraUtilsMPEGVideo.cpp
--- a/Engine/Plugins/Media/ElectraCodecs/Source/ElectraDecoders/Private/Utils/MPEG/ElectraUtilsMPEGVideo.cpp
+++ b/Engine/Plugins/Media/ElectraCodecs/Source/ElectraDecoders/Private/Utils/MPEG/ElectraUtilsMPEGVideo.cpp
@@ -11,684 +11,9 @@ namespace ElectraDecodersUtil
 {
 	namespace MPEG
 	{
-		void FAVCDecoderConfigurationRecord::SetRawData(const void* Data, int64 Size)
-		{
-			RawData.Empty();
-			if (Size)
-			{
-				RawData.Reserve((uint32)Size);
-				RawData.SetNumUninitialized((uint32)Size);
-				FMemory::Memcpy(RawData.GetData(), Data, Size);
-			}
-		}
-
-		const TArray<uint8>& FAVCDecoderConfigurationRecord::GetRawData() const
-		{
-			return RawData;
-		}
-
-		const TArray<uint8>& FAVCDecoderConfigurationRecord::GetCodecSpecificData() const
-		{
-			return CodecSpecificData;
-		}
-
-		const TArray<uint8>& FAVCDecoderConfigurationRecord::GetCodecSpecificDataSPS() const
-		{
-			return CodecSpecificDataSPSOnly;
-		}
-
-		const TArray<uint8>& FAVCDecoderConfigurationRecord::GetCodecSpecificDataPPS() const
-		{
-			return CodecSpecificDataPPSOnly;
-		}
-
-		int32 FAVCDecoderConfigurationRecord::GetNumberOfSPS() const
-		{
-			return ParsedSPSs.Num();
-		}
-
-		const FISO14496_10_seq_parameter_set_data& FAVCDecoderConfigurationRecord::GetParsedSPS(int32 SpsIndex) const
-		{
-			check(SpsIndex < GetNumberOfSPS());
-			return ParsedSPSs[SpsIndex];
-		}
-
-
-		bool FAVCDecoderConfigurationRecord::CreateFromCodecSpecificData(const TArray<uint8>& InFromCSD)
-		{
-			if (!InFromCSD.IsEmpty())
-			{
-				CodecSpecificData = InFromCSD;
-				TArray<ElectraDecodersUtil::MPEG::H264::FNaluInfo> Nalus;
-				if (!ElectraDecodersUtil::MPEG::H264::ParseBitstreamForNALUs(Nalus, InFromCSD.GetData(), InFromCSD.Num()))
-				{
-					return false;
-				}
-
-				FArray aSPS, aPPS;
-				for(int32 i=0; i<Nalus.Num(); ++i)
-				{
-					if (Nalus[i].Type == 7)
-					{
-						TMap<uint32, ElectraDecodersUtil::MPEG::H264::FSequenceParameterSet> SPSs;
-						if (ElectraDecodersUtil::MPEG::H264::ParseSequenceParameterSet(SPSs, InFromCSD.GetData() + Nalus[i].Offset + Nalus[i].UnitLength, Nalus[i].Size))
-						{
-							ElectraDecodersUtil::MPEG::H264::FSequenceParameterSet sps(SPSs.CreateConstIterator().Value());
-							AVCProfileIndication = sps.profile_idc;
-							ProfileCompatibility = (sps.constraint_set0_flag << 7) | (sps.constraint_set1_flag << 6) | (sps.constraint_set2_flag << 5) |
-												   (sps.constraint_set3_flag << 4) | (sps.constraint_set4_flag << 3) | (sps.constraint_set5_flag << 2);
-							AVCLevelIndication = sps.level_idc;
-							NALUnitLength = 4;
-							ChromaFormat = sps.chroma_format_idc;
-							BitDepthLumaMinus8 = sps.bit_depth_luma_minus8;
-							BitDepthChromaMinus8 = sps.bit_depth_chroma_minus8;
-							bHaveAdditionalProfileIndication = false;
-						}
-						else
-						{
-							return false;
-						}
-						aSPS.NALUs.Emplace_GetRef() = MakeArrayView<const uint8>(InFromCSD.GetData() + Nalus[i].Offset + Nalus[i].UnitLength, Nalus[i].Size);
-					}
-					else if (Nalus[i].Type == 8)
-					{
-						aPPS.NALUs.Emplace_GetRef() = MakeArrayView<const uint8>(InFromCSD.GetData() + Nalus[i].Offset + Nalus[i].UnitLength, Nalus[i].Size);
-					}
-				}
-				if (aSPS.NALUs.IsEmpty() || aPPS.NALUs.IsEmpty())
-				{
-					return false;
-				}
-
-				// Create the raw configuration record now.
-				ConfigurationVersion = 1;
-				FElectraBitstreamWriter wr;
-				wr.PutBits(ConfigurationVersion, 8);
-				wr.PutBits(AVCProfileIndication, 8);
-				wr.PutBits(ProfileCompatibility, 8);
-				wr.PutBits(AVCLevelIndication, 8);
-				wr.PutBits(63U, 6);
-				wr.PutBits(NALUnitLength-1, 2);
-				wr.PutBits(7U, 3);
-				wr.PutBits(static_cast<uint32>(aSPS.NALUs.Num()), 5);
-				for(int32 j=0; j<aSPS.NALUs.Num(); ++j)
-				{
-					const TArray<uint8>& d(aSPS.NALUs[j]);
-					wr.PutBits(static_cast<uint32>(d.Num()), 16);
-					for(int32 k=0; k<d.Num(); ++k)
-					{
-						wr.PutBits(d[k], 8);
-					}
-				}
-				wr.PutBits(static_cast<uint32>(aPPS.NALUs.Num()), 8);
-				for(int32 j=0; j<aPPS.NALUs.Num(); ++j)
-				{
-					const TArray<uint8>& d(aPPS.NALUs[j]);
-					wr.PutBits(static_cast<uint32>(d.Num()), 16);
-					for(int32 k=0; k<d.Num(); ++k)
-					{
-						wr.PutBits(d[k], 8);
-					}
-				}
-				if (AVCProfileIndication != 66 && AVCProfileIndication != 77 && AVCProfileIndication != 88)
-				{
-					wr.PutBits(63U, 6);
-					wr.PutBits(ChromaFormat, 2);
-					wr.PutBits(31U, 5);
-					wr.PutBits(BitDepthLumaMinus8, 3);
-					wr.PutBits(31U, 5);
-					wr.PutBits(BitDepthChromaMinus8, 3);
-					wr.PutBits(0U, 8);
-				}
-				wr.GetArray(RawData);
-				return true;
-			}
-			return false;
-		}
-
-		bool FAVCDecoderConfigurationRecord::Parse()
-		{
-			CodecSpecificData.Empty();
-
-			if (RawData.Num())
-			{
-				FElectraByteReader ByteReader(RawData.GetData(), RawData.Num());
-
-				if (!ByteReader.ReadByte(ConfigurationVersion))
-				{
-					return false;
-				}
-				if (!ByteReader.ReadByte(AVCProfileIndication))
-				{
-					return false;
-				}
-				if (!ByteReader.ReadByte(ProfileCompatibility))
-				{
-					return false;
-				}
-				if (!ByteReader.ReadByte(AVCLevelIndication))
-				{
-					return false;
-				}
-				if (!ByteReader.ReadByte(NALUnitLength))
-				{
-					return false;
-				}
-				NALUnitLength = (NALUnitLength & 3) + 1;
-				uint8 nSPS = 0;
-				if (!ByteReader.ReadByte(nSPS))
-				{
-					return false;
-				}
-				nSPS &= 31;
-				if (nSPS)
-				{
-					SequenceParameterSets.Reserve(nSPS);
-					ParsedSPSs.Reserve(nSPS);
-				}
-				int32 TotalSPSSize = 0;
-				for(int32 i = 0; i < nSPS; ++i)
-				{
-					uint16 spsLen = 0;
-					if (!ByteReader.ReadByte(spsLen))
-					{
-						return false;
-					}
-					TArray<uint8>& sps = SequenceParameterSets.AddDefaulted_GetRef();
-					sps.Reserve(spsLen);
-					sps.SetNumUninitialized(spsLen);
-					if (!ByteReader.ReadBytes(sps.GetData(), spsLen))
-					{
-						return false;
-					}
-					TotalSPSSize += 4 + spsLen;		// 4 because we always use 32 bit startcode and not NALUnitLength
-					ParseH264SPS(ParsedSPSs.AddDefaulted_GetRef(), sps.GetData(), sps.Num());
-				}
-				uint8 nPPS = 0;
-				if (!ByteReader.ReadByte(nPPS))
-				{
-					return false;
-				}
-				if (nPPS)
-				{
-					PictureParameterSets.Reserve(nPPS);
-				}
-				int32 TotalPPSSize = 0;
-				for(int32 i = 0; i < nPPS; ++i)
-				{
-					uint16 ppsLen = 0;
-					if (!ByteReader.ReadByte(ppsLen))
-					{
-						return false;
-					}
-					TArray<uint8>& pps = PictureParameterSets.AddDefaulted_GetRef();
-					pps.Reserve(ppsLen);
-					pps.SetNumUninitialized(ppsLen);
-					if (!ByteReader.ReadBytes(pps.GetData(), ppsLen))
-					{
-						return false;
-					}
-					TotalPPSSize += 4 + ppsLen;		// 4 because we always use 32 bit startcode and not NALUnitLength
-				}
-
-				if (AVCProfileIndication == 100 || AVCProfileIndication == 110 || AVCProfileIndication == 122 || AVCProfileIndication == 144)
-				{
-					// At least according to the ISO 14496-15:2014 standard these values must appear.
-					// I do have however files that are of AVC profile 100 but omit these values.
-					// Therefore let's do a quick check if we can read at least 4 more bytes.
-					if (ByteReader.BytesRemaining() >= 4)
-					{
-						if (!ByteReader.ReadByte(ChromaFormat))
-						{
-							return false;
-						}
-						ChromaFormat &= 3;
-						if (!ByteReader.ReadByte(BitDepthLumaMinus8))
-						{
-							return false;
-						}
-						BitDepthLumaMinus8 &= 7;
-						if (!ByteReader.ReadByte(BitDepthChromaMinus8))
-						{
-							return false;
-						}
-						BitDepthChromaMinus8 &= 7;
-						if (!ByteReader.ReadByte(nSPS))
-						{
-							return false;
-						}
-						if (nSPS)
-						{
-							SequenceParameterSetsExt.Reserve(nSPS);
-						}
-						for(int32 i = 0; i < nSPS; ++i)
-						{
-							uint16 spsLenExt = 0;
-							if (!ByteReader.ReadByte(spsLenExt))
-							{
-								return false;
-							}
-							TArray<uint8>& spsExt = SequenceParameterSetsExt.AddDefaulted_GetRef();
-							spsExt.Reserve(spsLenExt);
-							spsExt.SetNumUninitialized(spsLenExt);
-							if (!ByteReader.ReadBytes(spsExt.GetData(), spsLenExt))
-							{
-								return false;
-							}
-						}
-						bHaveAdditionalProfileIndication = true;
-					}
-				}
-
-				int32 TotalCSDSize = TotalSPSSize + TotalPPSSize;
-				if (TotalCSDSize)
-				{
-					CodecSpecificData.Reserve(TotalCSDSize);
-					CodecSpecificDataSPSOnly.Reserve(TotalSPSSize);
-					for(int32 i = 0; i < SequenceParameterSets.Num(); ++i)
-					{
-						CodecSpecificDataSPSOnly.Push(0);
-						CodecSpecificDataSPSOnly.Push(0);
-						CodecSpecificDataSPSOnly.Push(0);
-						CodecSpecificDataSPSOnly.Push(1);
-						CodecSpecificData.Push(0);
-						CodecSpecificData.Push(0);
-						CodecSpecificData.Push(0);
-						CodecSpecificData.Push(1);
-						for(int32 j = 0, jMax = SequenceParameterSets[i].Num(); j < jMax; ++j)
-						{
-							CodecSpecificDataSPSOnly.Push((SequenceParameterSets[i].GetData())[j]);
-							CodecSpecificData.Push((SequenceParameterSets[i].GetData())[j]);
-						}
-					}
-					CodecSpecificDataPPSOnly.Reserve(TotalPPSSize);
-					for(int32 i = 0; i < PictureParameterSets.Num(); ++i)
-					{
-						CodecSpecificDataPPSOnly.Push(0);
-						CodecSpecificDataPPSOnly.Push(0);
-						CodecSpecificDataPPSOnly.Push(0);
-						CodecSpecificDataPPSOnly.Push(1);
-						CodecSpecificData.Push(0);
-						CodecSpecificData.Push(0);
-						CodecSpecificData.Push(0);
-						CodecSpecificData.Push(1);
-						for(int32 j = 0, jMax = PictureParameterSets[i].Num(); j < jMax; ++j)
-						{
-							CodecSpecificDataPPSOnly.Push((PictureParameterSets[i].GetData())[j]);
-							CodecSpecificData.Push((PictureParameterSets[i].GetData())[j]);
-						}
-					}
-				}
-
-				return true;
-			}
-			return false;
-		}
-
-
-
-		void FHEVCDecoderConfigurationRecord::SetRawData(const void* Data, int64 Size)
-		{
-			RawData.Empty();
-			if (Size)
-			{
-				RawData.Reserve((uint32)Size);
-				RawData.SetNumUninitialized((uint32)Size);
-				FMemory::Memcpy(RawData.GetData(), Data, Size);
-			}
-		}
-
-		const TArray<uint8>& FHEVCDecoderConfigurationRecord::GetRawData() const
-		{
-			return RawData;
-		}
-
-		const TArray<uint8>& FHEVCDecoderConfigurationRecord::GetCodecSpecificData() const
-		{
-			return CodecSpecificData;
-		}
-
-		int32 FHEVCDecoderConfigurationRecord::GetNumberOfSPS() const
-		{
-			return ParsedSPSs.Num();
-		}
-
-		const FISO23008_2_seq_parameter_set_data& FHEVCDecoderConfigurationRecord::GetParsedSPS(int32 SpsIndex) const
-		{
-			return ParsedSPSs[SpsIndex];
-		}
-
-		void FHEVCDecoderConfigurationRecord::Reset()
-		{
-			RawData.Empty();
-			CodecSpecificData.Empty();
-			Arrays.Empty();
-			ConfigurationVersion = 0;
-			GeneralProfileSpace = 0;
-			GeneralTierFlag = 0;
-			GeneralProfileIDC = 0;
-			GeneralProfileCompatibilityFlags = 0;
-			GeneralConstraintIndicatorFlags = 0;
-			GeneralLevelIDC = 0;
-			MinSpatialSegmentationIDC = 0;
-			ParallelismType = 0;
-			ChromaFormat = 0;
-			BitDepthLumaMinus8 = 0;
-			BitDepthChromaMinus8 = 0;
-			AverageFrameRate = 0;
-			ConstantFrameRate = 0;
-			NumTemporalLayers = 0;
-			TemporalIdNested = 0;
-			NALUnitLengthMinus1 = 0;
-			ParsedSPSs.Empty();
-		}
-
-		bool FHEVCDecoderConfigurationRecord::CreateFromCodecSpecificData(const TArray<uint8>& InFromCSD)
-		{
-			if (!InFromCSD.IsEmpty())
-			{
-				CodecSpecificData = InFromCSD;
-				TArray<ElectraDecodersUtil::MPEG::H265::FNaluInfo> Nalus;
-				if (!ElectraDecodersUtil::MPEG::H265::ParseBitstreamForNALUs(Nalus, InFromCSD.GetData(), InFromCSD.Num()))
-				{
-					return false;
-				}
-				FArray aVPS, aSPS, aPPS;
-				for(int32 i=0; i<Nalus.Num(); ++i)
-				{
-					if (Nalus[i].Type == 32)
-					{
-						aVPS.NALUnitType = 32;
-						aVPS.NALUs.Emplace_GetRef() = MakeArrayView<const uint8>(InFromCSD.GetData() + Nalus[i].Offset + Nalus[i].UnitLength, Nalus[i].Size);
-					}
-					else if (Nalus[i].Type == 33)
-					{
-						TMap<uint32, ElectraDecodersUtil::MPEG::H265::FSequenceParameterSet> SPSs;
-						if (ElectraDecodersUtil::MPEG::H265::ParseSequenceParameterSet(SPSs, InFromCSD.GetData() + Nalus[i].Offset + Nalus[i].UnitLength, Nalus[i].Size))
-						{
-							ElectraDecodersUtil::MPEG::H265::FSequenceParameterSet sps(SPSs.CreateConstIterator().Value());
-							GeneralProfileSpace = sps.profile_tier_level.general_profile_space;
-							GeneralTierFlag = sps.profile_tier_level.general_tier_flag;
-							GeneralProfileIDC = sps.profile_tier_level.general_profile_idc;
-							GeneralProfileCompatibilityFlags = sps.profile_tier_level.general_profile_compatibility_flags;
-							GeneralConstraintIndicatorFlags = sps.profile_tier_level.general_constraint_indicator_flags;
-							GeneralLevelIDC = sps.profile_tier_level.general_level_idc;
-							check(sps.vui_parameters_present_flag);
-							MinSpatialSegmentationIDC = sps.vui_parameters.min_spatial_segmentation_idc;
-							ParallelismType = 0;	// for simplicity as 0 indicates `unknown`.
-							ChromaFormat = sps.chroma_format_idc;
-							BitDepthLumaMinus8 = sps.bit_depth_luma_minus8;
-							BitDepthChromaMinus8 = sps.bit_depth_chroma_minus8;
-							AverageFrameRate = 0;	// unknown
-							ConstantFrameRate = 0;	// may or may not
-							NumTemporalLayers = 0;	// unknown
-							TemporalIdNested = 0;	// may or may not
-							NALUnitLengthMinus1 = 3;
-						}
-						else
-						{
-							return false;
-						}
-						aSPS.NALUnitType = 33;
-						aSPS.NALUs.Emplace_GetRef() = MakeArrayView<const uint8>(InFromCSD.GetData() + Nalus[i].Offset + Nalus[i].UnitLength, Nalus[i].Size);
-					}
-					else if (Nalus[i].Type == 34)
-					{
-						aPPS.NALUnitType = 34;
-						aPPS.NALUs.Emplace_GetRef() = MakeArrayView<const uint8>(InFromCSD.GetData() + Nalus[i].Offset + Nalus[i].UnitLength, Nalus[i].Size);
-					}
-				}
-				if (aVPS.NALUs.IsEmpty() || aSPS.NALUs.IsEmpty() || aPPS.NALUs.IsEmpty())
-				{
-					return false;
-				}
-				Arrays.Emplace(MoveTemp(aVPS));
-				Arrays.Emplace(MoveTemp(aSPS));
-				Arrays.Emplace(MoveTemp(aPPS));
-
-				// Create the raw configuration record now.
-				ConfigurationVersion = 1;
-				FElectraBitstreamWriter wr;
-				wr.PutBits(ConfigurationVersion, 8);
-				wr.PutBits(GeneralProfileSpace, 2);
-				wr.PutBits(GeneralTierFlag, 1);
-				wr.PutBits(GeneralProfileIDC, 5);
-				wr.PutBits(GeneralProfileCompatibilityFlags, 32);
-				wr.PutBits64(GeneralConstraintIndicatorFlags, 48);
-				wr.PutBits(GeneralLevelIDC, 8);
-				wr.PutBits(15U, 4);
-				wr.PutBits(MinSpatialSegmentationIDC, 12);
-				wr.PutBits(63U, 6);
-				wr.PutBits(ParallelismType, 2);
-				wr.PutBits(63U, 6);
-				wr.PutBits(ChromaFormat, 2);
-				wr.PutBits(31U, 5);
-				wr.PutBits(BitDepthLumaMinus8, 3);
-				wr.PutBits(31U, 5);
-				wr.PutBits(BitDepthChromaMinus8, 3);
-				wr.PutBits(AverageFrameRate, 16);
-				wr.PutBits(ConstantFrameRate, 2);
-				wr.PutBits(NumTemporalLayers, 3);
-				wr.PutBits(TemporalIdNested, 1);
-				wr.PutBits(NALUnitLengthMinus1, 2);
-				wr.PutBits(static_cast<uint32>(Arrays.Num()), 8);
-				for(int32 i=0; i<Arrays.Num(); ++i)
-				{
-					const FArray& a(Arrays[i]);
-					wr.PutBits(a.Completeness, 1);
-					wr.PutBits(0U, 1);
-					wr.PutBits(a.NALUnitType, 6);
-					wr.PutBits(static_cast<uint32>(a.NALUs.Num()), 16);
-					for(int32 j=0; j<a.NALUs.Num(); ++j)
-					{
-						const TArray<uint8>& d(a.NALUs[j]);
-						wr.PutBits(static_cast<uint32>(d.Num()), 16);
-						for(int32 k=0; k<d.Num(); ++k)
-						{
-							wr.PutBits(d[k], 8);
-						}
-					}
-				}
-				wr.GetArray(RawData);
-				return true;
-			}
-			return false;
-		}
-
-		bool FHEVCDecoderConfigurationRecord::Parse()
-		{
-			CodecSpecificData.Empty();
-
-			if (RawData.Num() >= 23)
-			{
-				FElectraBitstreamReader BitReader(RawData.GetData(), RawData.Num());
-
-				ConfigurationVersion = (uint8)BitReader.GetBits(8);
-				if (ConfigurationVersion != 1)
-				{
-					return false;
-				}
-				GeneralProfileSpace = (uint8)BitReader.GetBits(2);
-				GeneralTierFlag = (uint8)BitReader.GetBits(1);
-				GeneralProfileIDC = (uint8)BitReader.GetBits(5);
-				GeneralProfileCompatibilityFlags  = BitReader.GetBits(32);
-				for(int32 i=0; i<48; ++i)
-				{
-					GeneralConstraintIndicatorFlags = (GeneralConstraintIndicatorFlags << 1U) | uint64(BitReader.GetBits(1));
-				}
-				GeneralLevelIDC = (uint8)BitReader.GetBits(8);
-				if (BitReader.GetBits(4) != 15)
-				{
-					return false;
-				}
-				MinSpatialSegmentationIDC = (uint16)BitReader.GetBits(12);
-				if (BitReader.GetBits(6) != 63)
-				{
-					return false;
-				}
-				ParallelismType = (uint8)BitReader.GetBits(2);
-				if (BitReader.GetBits(6) != 63)
-				{
-					return false;
-				}
-				ChromaFormat = (uint8)BitReader.GetBits(2);
-				if (BitReader.GetBits(5) != 31)
-				{
-					return false;
-				}
-				BitDepthLumaMinus8 = (uint8)BitReader.GetBits(3);
-				if (BitReader.GetBits(5) != 31)
-				{
-					return false;
-				}
-				BitDepthChromaMinus8 = (uint8)BitReader.GetBits(3);
-				AverageFrameRate = (uint16)BitReader.GetBits(16);
-				ConstantFrameRate = (uint8)BitReader.GetBits(2);
-				NumTemporalLayers = (uint8)BitReader.GetBits(3);
-				TemporalIdNested = (uint8)BitReader.GetBits(1);
-				NALUnitLengthMinus1 = (uint8)BitReader.GetBits(2);
-				uint32 numArrays = BitReader.GetBits(8);
-				Arrays.Reserve(numArrays);
-				for(uint32 i=0; i<numArrays; ++i)
-				{
-					FArray &a = Arrays.AddDefaulted_GetRef();
-					a.Completeness = (uint8)BitReader.GetBits(1);
-					if (BitReader.GetBits(1) != 0)
-					{
-						return false;
-					}
-					a.NALUnitType = (uint8)BitReader.GetBits(6);
-					uint32 numNALUs = BitReader.GetBits(16);
-					a.NALUs.Reserve(numNALUs);
-					for(uint32 j=0; j<numNALUs; ++j)
-					{
-						uint16 naluLen = (uint16)BitReader.GetBits(16);
-						TArray<uint8> &n = a.NALUs.AddDefaulted_GetRef();
-						n.Reserve(naluLen);
-						for(uint32 k=0; k<naluLen; ++k)
-						{
-							n.Emplace((uint8)BitReader.GetBits(8));
-						}
-					}
-				}
-				if (BitReader.GetRemainingBits() != 0)
-				{
-					return false;
-				}
-
-				// CSD
-				CodecSpecificData.Empty();
-				for(int32 i=0; i<Arrays.Num(); ++i)
-				{
-					const FArray& a = Arrays[i];
-					// SPS nut?
-					if (a.NALUnitType == 33)
-					{
-						FISO23008_2_seq_parameter_set_data sps;
-						if (ParseH265SPS(sps, a.NALUs[0].GetData(), a.NALUs[0].Num()))
-						{
-							ParsedSPSs.Emplace(MoveTemp(sps));
-						}
-					}
-
-					for(int32 j=0; j<a.NALUs.Num(); ++j)
-					{
-						CodecSpecificData.Add(0);
-						CodecSpecificData.Add(0);
-						CodecSpecificData.Add(0);
-						CodecSpecificData.Add(1);
-						CodecSpecificData.Append(a.NALUs[j]);
-					}
-				}
-				return true;
-			}
-			return false;
-		}
-
-
-
-		static int32 FindStartCode(const uint8* InData, SIZE_T InDataSize, int32& NALUnitLength)
-		{
-			for(const uint8* Data = InData; InDataSize >= 3; ++Data, --InDataSize)
-			{
-				if (Data[0] == 0 && Data[1] == 0 && (Data[2] == 1 || (InDataSize >= 4 && Data[2] == 0 && Data[3] == 1)))
-				{
-					NALUnitLength = Data[2] ? 3 : 4;
-					return Data - reinterpret_cast<const uint8*>(InData);
-				}
-			}
-			NALUnitLength = -1;
-			return -1;
-		}
 
 
 
-		void ParseBitstreamForNALUs(TArray<FNaluInfo>& outNALUs, const void* InBitstream, uint64 InBitstreamLength)
-		{
-			outNALUs.Reset();
-
-			uint64 Pos = 0;
-			uint64 BytesToGo = InBitstreamLength;
-			const uint8* BitstreamData = (const uint8*)InBitstream;
-			while(1)
-			{
-				int32 UnitLength = 0, StartCodePos = FindStartCode(BitstreamData, BytesToGo, UnitLength);
-				if (StartCodePos >= 0)
-				{
-					if (outNALUs.Num())
-					{
-						outNALUs.Last().Size = StartCodePos;
-						outNALUs.Last().Type = *BitstreamData;
-					}
-					FNaluInfo n;
-					n.Offset = Pos + StartCodePos;
-					n.Size = 0;
-					n.Type = 0;
-					n.UnitLength = UnitLength;
-					outNALUs.Push(n);
-					BitstreamData = AdvancePointer(BitstreamData, StartCodePos + UnitLength);
-					Pos += StartCodePos + UnitLength;
-					BytesToGo -= StartCodePos + UnitLength;
-				}
-				else
-				{
-					if (outNALUs.Num())
-					{
-						outNALUs.Last().Size = BytesToGo;
-						outNALUs.Last().Type = *BitstreamData;
-					}
-					break;
-				}
-			}
-
-		}
-
-
-
-		int32 EBSPtoRBSP(uint8* OutBuf, const uint8* InBuf, int32 NumBytesIn)
-		{
-			uint8* OutBase = OutBuf;
-			while(NumBytesIn-- > 0)
-			{
-				uint8 b = *InBuf++;
-				*OutBuf++ = b;
-				if (b == 0)
-				{
-					if (NumBytesIn > 1)
-					{
-						if (InBuf[0] == 0x00 && InBuf[1] == 0x03)
-						{
-							*OutBuf++ = 0x00;
-							InBuf += 2;
-							NumBytesIn -= 2;
-						}
-					}
-				}
-			}
-			return OutBuf - OutBase;
-		}
-
-
 
 		template <typename T>
 		struct FScopedDataPtr
@@ -709,553 +34,34 @@ namespace ElectraDecodersUtil
 		};
 
 
-		bool ParseH264SPS(FISO14496_10_seq_parameter_set_data& OutSPS, const void* Data, int32 Size)
-		{
-			struct FSyntaxElement
-			{
-				static uint32 ue_v(FElectraBitstreamReader& BitStream)
-				{
-					int32 lz = -1;
-					for(uint32 b = 0; b == 0; ++lz)
-					{
-						b = BitStream.GetBits(1);
-					}
-					if (lz)
-					{
-						return ((1 << lz) | BitStream.GetBits(lz)) - 1;
-					}
-					return 0;
-				}
-				static int32 se_v(FElectraBitstreamReader& BitStream)
-				{
-					uint32 c = ue_v(BitStream);
-					return c & 1 ? int32((c + 1) >> 1) : -int32((c + 1) >> 1);
-				}
-			};
-
-			// SPS is usually an EBSP so we need to strip it down.
-			FScopedDataPtr<uint8> RBSP(FMemory::Malloc(Size));
-			int32 RBSPsize = EBSPtoRBSP(RBSP, static_cast<const uint8*>(Data), Size);
-			FElectraBitstreamReader BitReader(RBSP, RBSPsize);
-
-			FMemory::Memzero(OutSPS);
-
-			uint8 nalUnitType = (uint8)BitReader.GetBits(8);
-			if ((nalUnitType & 0x1f) != 0x7)	// SPS NALU?
-			{
-				return false;
-			}
-
-			OutSPS.profile_idc = (uint8)BitReader.GetBits(8);
-			OutSPS.constraint_set0_flag = (uint8)BitReader.GetBits(1);
-			OutSPS.constraint_set1_flag = (uint8)BitReader.GetBits(1);
-			OutSPS.constraint_set2_flag = (uint8)BitReader.GetBits(1);
-			OutSPS.constraint_set3_flag = (uint8)BitReader.GetBits(1);
-			OutSPS.constraint_set4_flag = (uint8)BitReader.GetBits(1);
-			OutSPS.constraint_set5_flag = (uint8)BitReader.GetBits(1);
-			BitReader.SkipBits(2);
-			OutSPS.level_idc = (uint8)BitReader.GetBits(8);
-			OutSPS.seq_parameter_set_id = FSyntaxElement::ue_v(BitReader);
-			if (OutSPS.profile_idc == 100 || OutSPS.profile_idc == 110 || OutSPS.profile_idc == 122 || OutSPS.profile_idc == 244 ||
-				OutSPS.profile_idc == 44 || OutSPS.profile_idc == 83 || OutSPS.profile_idc == 86 || OutSPS.profile_idc == 118 || OutSPS.profile_idc == 128 ||
-				OutSPS.profile_idc == 138 || OutSPS.profile_idc == 139 || OutSPS.profile_idc == 134 || OutSPS.profile_idc == 135)
-			{
-				OutSPS.chroma_format_idc = FSyntaxElement::ue_v(BitReader);
-				if (OutSPS.chroma_format_idc == 3)
-				{
-					OutSPS.separate_colour_plane_flag = (uint8)BitReader.GetBits(1);
-				}
-				OutSPS.bit_depth_luma_minus8 = FSyntaxElement::ue_v(BitReader);
-				OutSPS.bit_depth_chroma_minus8 = FSyntaxElement::ue_v(BitReader);
-				OutSPS.qpprime_y_zero_transform_bypass_flag = (uint8)BitReader.GetBits(1);
-				OutSPS.seq_scaling_matrix_present_flag = (uint8)BitReader.GetBits(1);
-				if (OutSPS.seq_scaling_matrix_present_flag)
-				{
-					auto scaling_list = [&BitReader](int32* scalingList, int32 sizeOfScalingList, bool& useDefaultScalingMatrixFlag) -> void
-					{
-						int32 lastScale = 8;
-						int32 nextScale = 8;
-						for(int32 j=0; j<sizeOfScalingList; ++j)
-						{
-							if (nextScale)
-							{
-								int32 delta_scale = FSyntaxElement::se_v(BitReader);
-								nextScale = (lastScale + delta_scale + 256) % 256;
-								useDefaultScalingMatrixFlag = (j == 0 && nextScale == 0);
-							}
-							scalingList[j] = (nextScale == 0) ? lastScale : nextScale;
-							lastScale = scalingList[j];
-						}
-					};
 
-					// Skip over the scaling matrices.
-					int32 dummyScalingMatrix[64] = {0};
-					bool bDummyDefaultScalingMatrixFlag = false;
-					for(int32 i=0, iMax=OutSPS.chroma_format_idc!=3?8:12; i<iMax; ++i)
-					{
-						uint8 seq_scaling_list_present_flag = (uint8)BitReader.GetBits(1);
-						if (seq_scaling_list_present_flag)
-						{
-							if (i < 6)
-							{
-								scaling_list(dummyScalingMatrix, 16, bDummyDefaultScalingMatrixFlag);
-							}
-							else
-							{
-								scaling_list(dummyScalingMatrix, 64, bDummyDefaultScalingMatrixFlag);
-							}
-						}
-					}
-				}
-			}
-			OutSPS.log2_max_frame_num_minus4 = FSyntaxElement::ue_v(BitReader);
-			OutSPS.pic_order_cnt_type = FSyntaxElement::ue_v(BitReader);
-			if (OutSPS.pic_order_cnt_type == 0)
-			{
-				OutSPS.log2_max_pic_order_cnt_lsb_minus4 = FSyntaxElement::ue_v(BitReader);
-			}
-			else if (OutSPS.pic_order_cnt_type == 1)
-			{
-				OutSPS.delta_pic_order_always_zero_flag = FSyntaxElement::ue_v(BitReader);
-				OutSPS.offset_for_non_ref_pic = FSyntaxElement::se_v(BitReader);
-				OutSPS.offset_for_top_to_bottom_field = FSyntaxElement::se_v(BitReader);
-				OutSPS.num_ref_frames_in_pic_order_cnt_cycle = FSyntaxElement::ue_v(BitReader);
-				for(uint32 i = 0; i < OutSPS.num_ref_frames_in_pic_order_cnt_cycle; ++i)
-				{
-					FSyntaxElement::se_v(BitReader);		// discard
-				}
-			}
-			OutSPS.max_num_ref_frames = FSyntaxElement::ue_v(BitReader);
-			OutSPS.gaps_in_frame_num_value_allowed_flag = (uint8)BitReader.GetBits(1);
-			OutSPS.pic_width_in_mbs_minus1 = FSyntaxElement::ue_v(BitReader);
-			OutSPS.pic_height_in_map_units_minus1 = FSyntaxElement::ue_v(BitReader);
-			OutSPS.frame_mbs_only_flag = (uint8)BitReader.GetBits(1);
-			if (!OutSPS.frame_mbs_only_flag)
-			{
-				OutSPS.mb_adaptive_frame_field_flag = (uint8)BitReader.GetBits(1);
-			}
-			OutSPS.direct_8x8_inference_flag = (uint8)BitReader.GetBits(1);
-			OutSPS.frame_cropping_flag = (uint8)BitReader.GetBits(1);
-			if (OutSPS.frame_cropping_flag)
-			{
-				OutSPS.frame_crop_left_offset = FSyntaxElement::ue_v(BitReader);
-				OutSPS.frame_crop_right_offset = FSyntaxElement::ue_v(BitReader);
-				OutSPS.frame_crop_top_offset = FSyntaxElement::ue_v(BitReader);
-				OutSPS.frame_crop_bottom_offset = FSyntaxElement::ue_v(BitReader);
-			}
-			OutSPS.vui_parameters_present_flag = (uint8)BitReader.GetBits(1);
-			if (OutSPS.vui_parameters_present_flag)
-			{
-				OutSPS.aspect_ratio_info_present_flag = (uint8)BitReader.GetBits(1);
-				if (OutSPS.aspect_ratio_info_present_flag)
-				{
-					OutSPS.aspect_ratio_idc = (uint8)BitReader.GetBits(8);
-					if (OutSPS.aspect_ratio_idc == 255)
-					{
-						OutSPS.sar_width = (uint16)BitReader.GetBits(16);
-						OutSPS.sar_height = (uint16)BitReader.GetBits(16);
-					}
-				}
-				OutSPS.overscan_info_present_flag = (uint8)BitReader.GetBits(1);
-				if (OutSPS.overscan_info_present_flag)
-				{
-					OutSPS.overscan_appropriate_flag = (uint8)BitReader.GetBits(1);
-				}
-				OutSPS.video_signal_type_present_flag = (uint8)BitReader.GetBits(1);
-				if (OutSPS.video_signal_type_present_flag)
-				{
-					OutSPS.video_format = (uint8)BitReader.GetBits(3);
-					OutSPS.video_full_range_flag = (uint8)BitReader.GetBits(1);
-					OutSPS.colour_description_present_flag = (uint8)BitReader.GetBits(1);
-					if (OutSPS.colour_description_present_flag)
-					{
-						OutSPS.colour_primaries = (uint8)BitReader.GetBits(8);
-						OutSPS.transfer_characteristics = (uint8)BitReader.GetBits(8);
-						OutSPS.matrix_coefficients = (uint8)BitReader.GetBits(8);
-					}
-				}
-				OutSPS.chroma_loc_info_present_flag = (uint8)BitReader.GetBits(1);
-				if (OutSPS.chroma_loc_info_present_flag)
-				{
-					OutSPS.chroma_sample_loc_type_top_field = FSyntaxElement::ue_v(BitReader);
-					OutSPS.chroma_sample_loc_type_bottom_field = FSyntaxElement::ue_v(BitReader);
-				}
-				OutSPS.timing_info_present_flag = (uint8)BitReader.GetBits(1);
-				if (OutSPS.timing_info_present_flag)
-				{
-					OutSPS.num_units_in_tick = BitReader.GetBits(32);
-					OutSPS.time_scale = BitReader.GetBits(32);
-					OutSPS.fixed_frame_rate_flag = (uint8)BitReader.GetBits(1);
-				}
-				// The remainder is of no interest to us at the moment.
-			}
-			return true;
-		}
 
 
-		//! Parses a H.265 (ISO/IEC 23008-2) SPS NALU.
-		bool ParseH265SPS(FISO23008_2_seq_parameter_set_data& OutSPS, const void* Data, int32 Size)
+		bool ExtractSEIMessages(TArray<FSEIMessage>& OutMessages, const void* InBitstream, uint64 InBitstreamLength, ESEIStreamType StreamType, bool bIsPrefixSEI)
 		{
-			struct FSyntaxElement
-			{
-				static uint32 ue_v(FElectraBitstreamReader& BitStream)
-				{
-					int32 lz = -1;
-					for(uint32 b = 0; b == 0; ++lz)
-					{
-						b = BitStream.GetBits(1);
-					}
-					if (lz)
-					{
-						return ((1 << lz) | BitStream.GetBits(lz)) - 1;
-					}
-					return 0;
-				}
-				static int32 se_v(FElectraBitstreamReader& BitStream)
-				{
-					uint32 c = ue_v(BitStream);
-					return c & 1 ? int32((c + 1) >> 1) : -int32((c + 1) >> 1);
-				}
-			};
-
-			// SPS is usually an EBSP so we need to strip it down.
-			FScopedDataPtr<uint8> RBSP(FMemory::Malloc(Size));
-			int32 RBSPsize = EBSPtoRBSP(RBSP, static_cast<const uint8*>(Data), Size);
-			FElectraBitstreamReader BitReader(RBSP, RBSPsize);
-
-			FMemory::Memzero(OutSPS);
-
-			if (BitReader.GetBits(1) != 0)	// forbidden_zero_bit
-			{
-				return false;
-			}
-			if (BitReader.GetBits(6) != 33)	// sps_nut ?
-			{
-				return false;
-			}
-			BitReader.SkipBits(6);			// nuh_layer_id
-			BitReader.SkipBits(3);			// nuh_temporal_id_plus1
-
-			OutSPS.sps_video_parameter_set_id = BitReader.GetBits(4);
-			OutSPS.sps_max_sub_layers_minus1 = BitReader.GetBits(3);
-			OutSPS.sps_temporal_id_nesting_flag = BitReader.GetBits(1);
-
-			// profile_tier_level( sps_max_sub_layers_minus1 )
-			OutSPS.general_profile_space = BitReader.GetBits(2);
-			OutSPS.general_tier_flag = BitReader.GetBits(1);
-			OutSPS.general_profile_idc = BitReader.GetBits(5);
-			for(int32 i=0; i<32; ++i)
-			{
-				OutSPS.general_profile_compatibility_flag = (OutSPS.general_profile_compatibility_flag << 1) | BitReader.GetBits(1);
-			}
-			OutSPS.general_progressive_source_flag = BitReader.GetBits(1);
-			OutSPS.general_interlaced_source_flag = BitReader.GetBits(1);
-			OutSPS.general_non_packed_constraint_flag = BitReader.GetBits(1);
-			OutSPS.general_frame_only_constraint_flag = BitReader.GetBits(1);
-			// We skip over the next 44 bits that vary in layout depending on general_profile_idc and general_profile_compatibility_flag.
-			// For the moment none of the flags interest us, so we do not parse them for simplicity.
-			OutSPS.general_reserved_zero_44bits = ((uint64)BitReader.GetBits(32) << 32) | BitReader.GetBits(12);
-
-			OutSPS.general_level_idc = BitReader.GetBits(8);
-			for(uint32 i=0; i<OutSPS.sps_max_sub_layers_minus1; ++i)
-			{
-				OutSPS.sub_layer_profile_present_flag[i] = BitReader.GetBits(1);
-				OutSPS.sub_layer_level_present_flag[i] = BitReader.GetBits(1);
-			}
-			if (OutSPS.sps_max_sub_layers_minus1 > 0)
+			auto EBSPtoRBSP = [](uint8* OutBuf, const uint8* InBuf, int32 NumBytesIn) -> int32
 			{
-				for(uint32 i = OutSPS.sps_max_sub_layers_minus1; i<8; ++i)
+				uint8* OutBase = OutBuf;
+				while(NumBytesIn-- > 0)
 				{
-					BitReader.SkipBits(2);		// reserved_zero_2bits[ i ]
-				}
-			}
-			for(uint32 i=0; i<OutSPS.sps_max_sub_layers_minus1; ++i)
-			{
-				// Similarly, we skip all of these flags that are of no interest to us right now.
-				if (OutSPS.sub_layer_profile_present_flag[i])
-				{
-					BitReader.SkipBits(2);		// sub_layer_profile_space[ i ]
-					BitReader.SkipBits(1);		// sub_layer_tier_flag[ i ]
-					BitReader.SkipBits(5);		// sub_layer_profile_idc[ i ]
-					BitReader.SkipBits(32);		// for( j = 0; j < 32; j++ ) sub_layer_profile_compatibility_flag[ i ][ j ] u(1)
-					BitReader.SkipBits(1);		// sub_layer_progressive_source_flag[ i ]
-					BitReader.SkipBits(1);		// sub_layer_interlaced_source_flag[ i ]
-					BitReader.SkipBits(1);		// sub_layer_non_packed_constraint_flag[ i ]
-					BitReader.SkipBits(1);		// sub_layer_frame_only_constraint_flag[ i ]
-					BitReader.SkipBits(44);		// sub_layer_reserved_zero_44bits[ i ]
-				}
-				if (OutSPS.sub_layer_level_present_flag[i])
-				{
-					BitReader.SkipBits(8);		// sub_layer_level_idc[ i ]
-				}
-			}
-
-			OutSPS.sps_seq_parameter_set_id = FSyntaxElement::ue_v(BitReader);
-			OutSPS.chroma_format_idc = FSyntaxElement::ue_v(BitReader);
-			if (OutSPS.chroma_format_idc == 3)
-			{
-				OutSPS.separate_colour_plane_flag = BitReader.GetBits(1);			// separate_colour_plane_flag
-			}
-			OutSPS.pic_width_in_luma_samples = FSyntaxElement::ue_v(BitReader);
-			OutSPS.pic_height_in_luma_samples = FSyntaxElement::ue_v(BitReader);
-
-			OutSPS.conformance_window_flag = BitReader.GetBits(1);
-			if (OutSPS.conformance_window_flag)
-			{
-				OutSPS.conf_win_left_offset = FSyntaxElement::ue_v(BitReader);
-				OutSPS.conf_win_right_offset = FSyntaxElement::ue_v(BitReader);
-				OutSPS.conf_win_top_offset = FSyntaxElement::ue_v(BitReader);
-				OutSPS.conf_win_bottom_offset = FSyntaxElement::ue_v(BitReader);
-			}
-			OutSPS.bit_depth_luma_minus8 = FSyntaxElement::ue_v(BitReader);
-			OutSPS.bit_depth_chroma_minus8 = FSyntaxElement::ue_v(BitReader);
-			OutSPS.log2_max_pic_order_cnt_lsb_minus4 = FSyntaxElement::ue_v(BitReader);
-			OutSPS.sps_sub_layer_ordering_info_present_flag = BitReader.GetBits(1);
-			for(uint32 i=(OutSPS.sps_sub_layer_ordering_info_present_flag ? 0 : OutSPS.sps_max_sub_layers_minus1); i<=OutSPS.sps_max_sub_layers_minus1; ++i)
-			{
-				OutSPS.sps_max_dec_pic_buffering_minus1[i] = FSyntaxElement::ue_v(BitReader);
-				OutSPS.sps_max_num_reorder_pics[i] = FSyntaxElement::ue_v(BitReader);
-				OutSPS.sps_max_latency_increase_plus1[i] = FSyntaxElement::ue_v(BitReader);
-			}
-			OutSPS.log2_min_luma_coding_block_size_minus3 = FSyntaxElement::ue_v(BitReader);
-			OutSPS.log2_diff_max_min_luma_coding_block_size = FSyntaxElement::ue_v(BitReader);
-			OutSPS.log2_min_transform_block_size_minus2 = FSyntaxElement::ue_v(BitReader);
-			OutSPS.log2_diff_max_min_transform_block_size = FSyntaxElement::ue_v(BitReader);
-			OutSPS.max_transform_hierarchy_depth_inter = FSyntaxElement::ue_v(BitReader);
-			OutSPS.max_transform_hierarchy_depth_intra = FSyntaxElement::ue_v(BitReader);
-			OutSPS.scaling_list_enabled_flag = BitReader.GetBits(1);
-			if (OutSPS.scaling_list_enabled_flag)
-			{
-				// Parse the scaling list into temp variables only. We do not need them right now.
-				OutSPS.sps_scaling_list_data_present_flag = BitReader.GetBits(1);
-				if (OutSPS.sps_scaling_list_data_present_flag)
-				{
-					// scaling_list_data( )
-					for(uint32 sizeId=0; sizeId<4; ++sizeId)
-					{
-						for(int32 matrixId=0; matrixId<((sizeId==3)?2:6); ++matrixId)
-						{
-							uint8 scaling_list_pred_mode_flag = BitReader.GetBits(1);	// [ sizeId ][ matrixId ] u(1)
-							if (!scaling_list_pred_mode_flag/*[sizeId][matrixId]*/)
-							{
-								uint32 scaling_list_pred_matrix_id_delta/*[sizeId][matrixId]*/ = FSyntaxElement::ue_v(BitReader);
-							}
-							else
-							{
-								//nextCoef = 8;
-								uint32 coefNum = Min((uint32)64U, (uint32)(1 << (4 + (sizeId << 1))));
-								if (sizeId > 1)
-								{
-									uint32 scaling_list_dc_coef_minus8/*[sizeId ? 2][matrixId]*/ = FSyntaxElement::se_v(BitReader);
-									//nextCoef = scaling_list_dc_coef_minus8[sizeId ? 2][matrixId] + 8;
-								}
-								for(uint32 i=0; i<coefNum; i++)
-								{
-									uint32 scaling_list_delta_coef = FSyntaxElement::se_v(BitReader);
-									//nextCoef = ( nextCoef + scaling_list_delta_coef + 256 ) % 256
-									//ScalingList[ sizeId ][ matrixId ][ i ] = nextCoef
-								}
-							}
-						}
-					}
-				}
-			}
-
-			OutSPS.amp_enabled_flag = BitReader.GetBits(1);
-			OutSPS.sample_adaptive_offset_enabled_flag = BitReader.GetBits(1);
-			OutSPS.pcm_enabled_flag = BitReader.GetBits(1);
-			if (OutSPS.pcm_enabled_flag)
-			{
-				OutSPS.pcm_sample_bit_depth_luma_minus1 = BitReader.GetBits(4);
-				OutSPS.pcm_sample_bit_depth_chroma_minus1 = BitReader.GetBits(4);
-				OutSPS.log2_min_pcm_luma_coding_block_size_minus3 = FSyntaxElement::ue_v(BitReader);
-				OutSPS.log2_diff_max_min_pcm_luma_coding_block_size = FSyntaxElement::ue_v(BitReader);
-				OutSPS.pcm_loop_filter_disabled_flag = BitReader.GetBits(1);
-			}
-
-			OutSPS.num_short_term_ref_pic_sets = FSyntaxElement::ue_v(BitReader);
-			uint32 unused_num_delta_pocs[64] = {0};
-			for(uint32 i=0; i<OutSPS.num_short_term_ref_pic_sets; ++i)
-			{
-				struct short_term_ref_pic_set
-				{
-					static bool parse(uint32 stRpsIdx, uint32 num_short_term_ref_pic_sets, uint32* num_delta_pocs, FElectraBitstreamReader& BitReader)
+					uint8 b = *InBuf++;
+					*OutBuf++ = b;
+					if (b == 0)
 					{
-						uint32 inter_ref_pic_set_prediction_flag = stRpsIdx != 0 ? BitReader.GetBits(1) : 0;
-						if (inter_ref_pic_set_prediction_flag)
-						{
-							uint32 delta_idx_minus1 = 0;
-							if (stRpsIdx == num_short_term_ref_pic_sets)
-							{
-								delta_idx_minus1 = FSyntaxElement::ue_v(BitReader);
-							}
-							BitReader.SkipBits(1);	// delta_rps_sign
-							uint32 abs_delta_rps_minus1 = FSyntaxElement::ue_v(BitReader);
-							uint32 RefRpsIdx = stRpsIdx - (delta_idx_minus1 + 1);
-							num_delta_pocs[stRpsIdx] = 0;
-							for(uint32 j=0; j<=num_delta_pocs[RefRpsIdx]; ++j)
-							{
-								uint8 used_by_curr_pic_flag/*[j]*/ = BitReader.GetBits(1);
-								uint8 use_delta_flag = 0;
-								if (!used_by_curr_pic_flag/*[j]*/)
-								{
-									use_delta_flag = BitReader.GetBits(1);
-								}
-								if (used_by_curr_pic_flag || use_delta_flag)
-								{
-									++num_delta_pocs[stRpsIdx];
-								}
-							}
-						}
-						else
+						if (NumBytesIn > 1)
 						{
-							uint32 num_negative_pics = FSyntaxElement::ue_v(BitReader);
-							uint32 num_positive_pics = FSyntaxElement::ue_v(BitReader);
-
-							if (((uint64)num_positive_pics + (uint64)num_negative_pics) * 2 > BitReader.GetRemainingBits())
-							{
-								return false;
-							}
-
-							num_delta_pocs[stRpsIdx] = num_negative_pics + num_positive_pics;
-
-							for(uint32 j=0; j<num_negative_pics; ++j)
+							if (InBuf[0] == 0x00 && InBuf[1] == 0x03)
 							{
-								uint32 delta_poc_s0_minus1/*[j]*/ = FSyntaxElement::ue_v(BitReader);
-								(void)delta_poc_s0_minus1;
-								BitReader.SkipBits(1);				// used_by_curr_pic_s0_flag[ j ]
-							}
-							for(uint32 j=0; j<num_positive_pics; ++j)
-							{
-								uint32 delta_poc_s1_minus1/*[j]*/ = FSyntaxElement::ue_v(BitReader);
-								(void)delta_poc_s1_minus1;
-								BitReader.SkipBits(1);				// used_by_curr_pic_s1_flag[ j ]
+								*OutBuf++ = 0x00;
+								InBuf += 2;
+								NumBytesIn -= 2;
 							}
 						}
-						return true;
 					}
-				};
-				if (!short_term_ref_pic_set::parse(i, OutSPS.num_short_term_ref_pic_sets, unused_num_delta_pocs, BitReader))
-				{
-					return false;
 				}
-			}
-
-			uint8 long_term_ref_pics_present_flag = BitReader.GetBits(1);
-			if (long_term_ref_pics_present_flag)
-			{
-				uint32 num_long_term_ref_pics_sps = FSyntaxElement::ue_v(BitReader);
-				for(uint32 i=0; i<num_long_term_ref_pics_sps; ++i)
-				{
-					uint32 lt_ref_pic_poc_lsb_sps = FSyntaxElement::ue_v(BitReader);	//[ i ] u(v)
-					BitReader.SkipBits(1);							//	used_by_curr_pic_lt_sps_flag[ i ]
-				}
-			}
-
-			OutSPS.sps_temporal_mvp_enabled_flag = BitReader.GetBits(1);
-			OutSPS.strong_intra_smoothing_enabled_flag = BitReader.GetBits(1);
-			OutSPS.vui_parameters_present_flag = BitReader.GetBits(1);
-			if (OutSPS.vui_parameters_present_flag)
-			{
-				// vui_parameters()
-				OutSPS.aspect_ratio_info_present_flag = BitReader.GetBits(1);
-				if (OutSPS.aspect_ratio_info_present_flag)
-				{
-					OutSPS.aspect_ratio_idc = BitReader.GetBits(8);
-					if (OutSPS.aspect_ratio_idc == 255)
-					{
-						OutSPS.sar_width = BitReader.GetBits(16);
-						OutSPS.sar_height = BitReader.GetBits(16);
-					}
-				}
-				OutSPS.overscan_info_present_flag = BitReader.GetBits(1);
-				if (OutSPS.overscan_info_present_flag)
-				{
-					OutSPS.overscan_appropriate_flag = BitReader.GetBits(1);
-				}
-				OutSPS.video_signal_type_present_flag = BitReader.GetBits(1);
-				if (OutSPS.video_signal_type_present_flag)
-				{
-					OutSPS.video_format = BitReader.GetBits(3);
-					OutSPS.video_full_range_flag = BitReader.GetBits(1);
-					OutSPS.colour_description_present_flag = BitReader.GetBits(1);
-					if (OutSPS.colour_description_present_flag)
-					{
-						OutSPS.colour_primaries = BitReader.GetBits(8);
-						OutSPS.transfer_characteristics = BitReader.GetBits(8);
-						OutSPS.matrix_coeffs = BitReader.GetBits(8);
-					}
-				}
-				OutSPS.chroma_loc_info_present_flag = BitReader.GetBits(1);
-				if (OutSPS.chroma_loc_info_present_flag)
-				{
-					OutSPS.chroma_sample_loc_type_top_field = FSyntaxElement::ue_v(BitReader);
-					OutSPS.chroma_sample_loc_type_bottom_field = FSyntaxElement::ue_v(BitReader);
-				}
-				OutSPS.neutral_chroma_indication_flag = BitReader.GetBits(1);
-				OutSPS.field_seq_flag = BitReader.GetBits(1);
-				OutSPS.frame_field_info_present_flag = BitReader.GetBits(1);
-				OutSPS.default_display_window_flag = BitReader.GetBits(1);
-				if (OutSPS.default_display_window_flag)
-				{
-					OutSPS.def_disp_win_left_offset = FSyntaxElement::ue_v(BitReader);
-					OutSPS.def_disp_win_right_offset = FSyntaxElement::ue_v(BitReader);
-					OutSPS.def_disp_win_top_offset = FSyntaxElement::ue_v(BitReader);
-					OutSPS.def_disp_win_bottom_offset = FSyntaxElement::ue_v(BitReader);
-				}
-				OutSPS.vui_timing_info_present_flag = BitReader.GetBits(1);
-				if (OutSPS.vui_timing_info_present_flag)
-				{
-					OutSPS.vui_num_units_in_tick = BitReader.GetBits(32);
-					OutSPS.vui_time_scale = BitReader.GetBits(32);
-					OutSPS.vui_poc_proportional_to_timing_flag = BitReader.GetBits(1);
-					if (OutSPS.vui_poc_proportional_to_timing_flag)
-					{
-						OutSPS.vui_num_ticks_poc_diff_one_minus1 = FSyntaxElement::ue_v(BitReader);
-// None of the following is of interest to us now, so we stop parsing at this point.
-					}
-				}
-			}
-			return true;
-		}
-
-
-		uint64 FISO23008_2_seq_parameter_set_data::GetConstraintFlags() const
-		{
-			uint64 ConstraintFlags = ((uint64)((general_progressive_source_flag << 3) | (general_interlaced_source_flag << 2) | (general_non_packed_constraint_flag << 1) | general_frame_only_constraint_flag) << 44) | general_reserved_zero_44bits;
-			return ConstraintFlags;
-		}
-
-		FString FISO23008_2_seq_parameter_set_data::GetRFC6381(const TCHAR* SampleTypePrefix) const
-		{
-			// As per ISO/IEC 14496-15:2014 Annex E.3
-			uint32 cf = BitReverse32(general_profile_compatibility_flag);
-			uint64 ConstraintFlags = GetConstraintFlags();
-			FString cfs;
-			bool bNonZero = false;
-			for(int32 i=5; i>=0; --i, ConstraintFlags >>= 8)
-			{
-				if ((ConstraintFlags & 255) == 0 && !bNonZero)
-				{
-					continue;
-				}
-				bNonZero = true;
-				cfs = FString::Printf(TEXT(".%02X"), (uint8)(ConstraintFlags & 255)) + cfs;
-			}
-
-			if (general_profile_space == 0)
-			{
-				return FString::Printf(TEXT("%s.%d.%X.%c%d%s"), SampleTypePrefix, general_profile_idc, cf, general_tier_flag ? TCHAR('H') : TCHAR('L'), general_level_idc, *cfs);
-			}
-			else
-			{
-				return FString::Printf(TEXT("%s.%c%d.%X.%c%d%s"), SampleTypePrefix, TCHAR(TCHAR('A')+general_profile_space-1), general_profile_idc, cf, general_tier_flag ? TCHAR('H') : TCHAR('L'), general_level_idc, *cfs);
-			}
-		}
-
+				return OutBuf - OutBase;
+			};
 
-		bool ExtractSEIMessages(TArray<FSEIMessage>& OutMessages, const void* InBitstream, uint64 InBitstreamLength, ESEIStreamType StreamType, bool bIsPrefixSEI)
-		{
 			FScopedDataPtr<uint8> RBSP(FMemory::Malloc((uint32) InBitstreamLength));
 			int32 RBSPsize = EBSPtoRBSP(RBSP, static_cast<const uint8*>(InBitstream), (int32)InBitstreamLength);
 			FElectraBitstreamReader BitReader(RBSP, RBSPsize);
diff --git a/Engine/Plugins/Media/ElectraCodecs/Source/ElectraDecoders/Public/Utils/MPEG/ElectraUtilsMPEGVideo.h b/Engine/Plugins/Media/ElectraCodecs/Source/ElectraDecoders/Public/Utils/MPEG/ElectraUtilsMPEGVideo.h
--- a/Engine/Plugins/Media/ElectraCodecs/Source/ElectraDecoders/Public/Utils/MPEG/ElectraUtilsMPEGVideo.h
+++ b/Engine/Plugins/Media/ElectraCodecs/Source/ElectraDecoders/Public/Utils/MPEG/ElectraUtilsMPEGVideo.h
@@ -11,447 +11,6 @@ namespace ElectraDecodersUtil
 {
 	namespace MPEG
 	{
-		using ElectraDecodersUtil::FFractionalValue;
-
-		/**
-		 * Sequence parameter set as per ISO/IEC 14496-10:2012, section 7.3.2.1.1
-		 */
-		struct FISO14496_10_seq_parameter_set_data
-		{
-			uint8	profile_idc = 0;
-			uint8	constraint_set0_flag = 0;
-			uint8	constraint_set1_flag = 0;
-			uint8	constraint_set2_flag = 0;
-			uint8	constraint_set3_flag = 0;
-			uint8	constraint_set4_flag = 0;
-			uint8	constraint_set5_flag = 0;
-			uint8	level_idc = 0;
-			uint32	seq_parameter_set_id = 0;
-			uint32	chroma_format_idc = 0;
-			uint8	separate_colour_plane_flag = 0;
-			uint32	bit_depth_luma_minus8 = 0;
-			uint32	bit_depth_chroma_minus8 = 0;
-			uint8	qpprime_y_zero_transform_bypass_flag = 0;
-			uint8	seq_scaling_matrix_present_flag = 0;
-			uint32	log2_max_frame_num_minus4 = 0;
-			uint32	pic_order_cnt_type = 0;
-			uint32	log2_max_pic_order_cnt_lsb_minus4 = 0;
-			uint32	delta_pic_order_always_zero_flag = 0;
-			int32	offset_for_non_ref_pic = 0;
-			int32	offset_for_top_to_bottom_field = 0;
-			uint32	num_ref_frames_in_pic_order_cnt_cycle = 0;
-			uint32	max_num_ref_frames = 0;
-			uint8	gaps_in_frame_num_value_allowed_flag = 0;
-			uint32	pic_width_in_mbs_minus1 = 0;
-			uint32	pic_height_in_map_units_minus1 = 0;
-			uint8	frame_mbs_only_flag = 0;
-			uint8	mb_adaptive_frame_field_flag = 0;
-			uint8	direct_8x8_inference_flag = 0;
-			uint8	frame_cropping_flag = 0;
-			uint32	frame_crop_left_offset = 0;
-			uint32	frame_crop_right_offset = 0;
-			uint32	frame_crop_top_offset = 0;
-			uint32	frame_crop_bottom_offset = 0;
-			uint8	vui_parameters_present_flag = 0;
-			uint8	aspect_ratio_info_present_flag = 0;
-			uint8	aspect_ratio_idc = 0;
-			uint16	sar_width = 0;
-			uint16	sar_height = 0;
-			uint8	overscan_info_present_flag = 0;
-			uint8	overscan_appropriate_flag = 0;
-			uint8	video_signal_type_present_flag = 0;
-			uint8	video_format = 0;
-			uint8	video_full_range_flag = 0;
-			uint8	colour_description_present_flag = 0;
-			uint8	colour_primaries = 0;
-			uint8	transfer_characteristics = 0;
-			uint8	matrix_coefficients = 0;
-			uint8	chroma_loc_info_present_flag = 0;
-			uint32	chroma_sample_loc_type_top_field = 0;
-			uint32	chroma_sample_loc_type_bottom_field = 0;
-			uint8	timing_info_present_flag = 0;
-			uint32	num_units_in_tick = 0;
-			uint32	time_scale = 0;
-			uint8	fixed_frame_rate_flag = 0;
-
-			int32 GetWidth() const
-			{
-				return (pic_width_in_mbs_minus1 + 1) * 16;
-			}
-
-			int32 GetHeight() const
-			{
-				return (pic_height_in_map_units_minus1 + 1) * 16;
-			}
-
-			void GetCrop(int32& Left, int32& Right, int32& Top, int32& Bottom) const
-			{
-				if (frame_cropping_flag)
-				{
-					// The scaling factors are determined by the chroma_format_idc (see ISO/IEC 14496-10, table 6.1)
-					// For our purposes this will be 1, so the sub width/height are 2.
-					const int32 CropUnitX = 2;
-					const int32 CropUnitY = 2;
-					Left = (uint16)frame_crop_left_offset * CropUnitX;
-					Right = (uint16)frame_crop_right_offset * CropUnitX;
-					Top = (uint16)frame_crop_top_offset * CropUnitY;
-					Bottom = (uint16)frame_crop_bottom_offset * CropUnitY;
-				}
-				else
-				{
-					Left = Right = Top = Bottom = 0;
-				}
-			}
-
-			void GetAspect(int32& SarW, int32& SarH) const
-			{
-				if (vui_parameters_present_flag && aspect_ratio_info_present_flag)
-				{
-					switch (aspect_ratio_idc)
-					{
-						default:	SarW = SarH = 0;		break;
-						case 0:		SarW = SarH = 0;		break;
-						case 1:		SarW = SarH = 1;		break;
-						case 2:		SarW = 12; SarH = 11;	break;
-						case 3:		SarW = 10; SarH = 11;	break;
-						case 4:		SarW = 16; SarH = 11;	break;
-						case 5:		SarW = 40; SarH = 33;	break;
-						case 6:		SarW = 24; SarH = 11;	break;
-						case 7:		SarW = 20; SarH = 11;	break;
-						case 8:		SarW = 32; SarH = 11;	break;
-						case 9:		SarW = 80; SarH = 33;	break;
-						case 10:	SarW = 18; SarH = 11;	break;
-						case 11:	SarW = 15; SarH = 11;	break;
-						case 12:	SarW = 64; SarH = 33;	break;
-						case 13:	SarW = 160; SarH = 99;	break;
-						case 14:	SarW = 4; SarH = 3;		break;
-						case 15:	SarW = 3; SarH = 2;		break;
-						case 16:	SarW = 2; SarH = 1;		break;
-						case 255:	SarW = sar_width; SarH = sar_height;	break;
-					}
-				}
-				else
-				{
-					SarW = SarH = 1;
-				}
-			}
-			FFractionalValue GetTiming() const
-			{
-				if (vui_parameters_present_flag && timing_info_present_flag)
-				{
-					return FFractionalValue(time_scale, num_units_in_tick * 2);
-				}
-				return FFractionalValue();
-			}
-		};
-
-		/**
-		 * ISO/IEC 14496-15:2014
-		 */
-		class ELECTRADECODERS_API FAVCDecoderConfigurationRecord
-		{
-		public:
-			void SetRawData(const void* Data, int64 Size);
-			const TArray<uint8>& GetRawData() const;
-			bool CreateFromCodecSpecificData(const TArray<uint8>& InFromCSD);
-
-			bool Parse();
-
-			const TArray<uint8>& GetCodecSpecificData() const;
-			const TArray<uint8>& GetCodecSpecificDataSPS() const;
-			const TArray<uint8>& GetCodecSpecificDataPPS() const;
-
-			int32 GetNumberOfSPS() const;
-			const FISO14496_10_seq_parameter_set_data& GetParsedSPS(int32 SpsIndex) const;
-
-			uint8 GetAVCProfileIndication() const
-			{ return AVCProfileIndication; }
-			uint8 GetProfileCompatibility() const
-			{ return ProfileCompatibility; }
-			uint8 GetAVCLevelIndication() const
-			{ return AVCLevelIndication; }
-
-		private:
-			struct FArray
-			{
-				TArray<TArray<uint8>> NALUs;
-			};
-
-			TArray<uint8>													RawData;
-			TArray<uint8>													CodecSpecificData;
-			TArray<uint8>													CodecSpecificDataSPSOnly;
-			TArray<uint8>													CodecSpecificDataPPSOnly;
-			TArray<TArray<uint8>>											SequenceParameterSets;
-			TArray<TArray<uint8>>											PictureParameterSets;
-			TArray<TArray<uint8>>											SequenceParameterSetsExt;
-			uint8															ConfigurationVersion = 0;
-			uint8															AVCProfileIndication = 0;
-			uint8															ProfileCompatibility = 0;
-			uint8															AVCLevelIndication = 0;
-			uint8															NALUnitLength = 0;
-			uint8															ChromaFormat = 0;
-			uint8															BitDepthLumaMinus8 = 0;
-			uint8															BitDepthChromaMinus8 = 0;
-			bool															bHaveAdditionalProfileIndication = false;
-
-			TArray<FISO14496_10_seq_parameter_set_data>						ParsedSPSs;
-		};
-
-
-		struct ELECTRADECODERS_API FISO23008_2_seq_parameter_set_data
-		{
-			uint8		sps_video_parameter_set_id = 0;
-			uint8		sps_max_sub_layers_minus1 = 0;
-			uint8		sps_temporal_id_nesting_flag = 0;
-			uint8		general_profile_space = 0;
-			uint8		general_tier_flag = 0;
-			uint8		general_profile_idc = 0;
-			uint32		general_profile_compatibility_flag = 0;
-			uint8		general_progressive_source_flag = 0;
-			uint8		general_interlaced_source_flag = 0;
-			uint8		general_non_packed_constraint_flag = 0;
-			uint8		general_frame_only_constraint_flag = 0;
-			uint64		general_reserved_zero_44bits = 0;
-			uint8		general_level_idc = 0;
-			uint8		sub_layer_profile_present_flag[8] = {0};
-			uint8		sub_layer_level_present_flag[8] = {0};
-			uint32		sps_seq_parameter_set_id = 0;
-			uint32		chroma_format_idc = 0;
-			uint8		separate_colour_plane_flag = 0;
-			uint32		pic_width_in_luma_samples = 0;
-			uint32		pic_height_in_luma_samples = 0;
-			uint8		conformance_window_flag = 0;
-			uint32		conf_win_left_offset = 0;
-			uint32		conf_win_right_offset = 0;
-			uint32		conf_win_top_offset = 0;
-			uint32		conf_win_bottom_offset = 0;
-			uint32		bit_depth_luma_minus8 = 0;
-			uint32		bit_depth_chroma_minus8 = 0;
-			uint32		log2_max_pic_order_cnt_lsb_minus4 = 0;
-			uint8		sps_sub_layer_ordering_info_present_flag = 0;
-			uint32		sps_max_dec_pic_buffering_minus1[8] = {0};
-			uint32		sps_max_num_reorder_pics[8] = {0};
-			uint32		sps_max_latency_increase_plus1[8] = {0};
-			uint32		log2_min_luma_coding_block_size_minus3 = 0;
-			uint32		log2_diff_max_min_luma_coding_block_size = 0;
-			uint32		log2_min_transform_block_size_minus2 = 0;
-			uint32		log2_diff_max_min_transform_block_size = 0;
-			uint32		max_transform_hierarchy_depth_inter = 0;
-			uint32		max_transform_hierarchy_depth_intra = 0;
-			uint8		scaling_list_enabled_flag = 0;
-			uint8		sps_scaling_list_data_present_flag = 0;
-			uint8		amp_enabled_flag = 0;
-			uint8		sample_adaptive_offset_enabled_flag = 0;
-			uint8		pcm_enabled_flag = 0;
-			uint8		pcm_sample_bit_depth_luma_minus1 = 0;
-			uint8		pcm_sample_bit_depth_chroma_minus1 = 0;
-			uint32		log2_min_pcm_luma_coding_block_size_minus3 = 0;
-			uint32		log2_diff_max_min_pcm_luma_coding_block_size = 0;
-			uint8		pcm_loop_filter_disabled_flag = 0;
-			uint32		num_short_term_ref_pic_sets = 0;
-			uint8		sps_temporal_mvp_enabled_flag = 0;
-			uint8		strong_intra_smoothing_enabled_flag = 0;
-			uint8		vui_parameters_present_flag = 0;
-			uint8		aspect_ratio_info_present_flag = 0;
-			uint8		aspect_ratio_idc = 0;
-			uint16		sar_width = 0;
-			uint16		sar_height = 0;
-			uint8		overscan_info_present_flag = 0;
-			uint8		overscan_appropriate_flag = 0;
-			uint8		video_signal_type_present_flag = 0;
-			uint8		video_format = 0;
-			uint8		video_full_range_flag = 0;
-			uint8		colour_description_present_flag = 0;
-			uint8		colour_primaries = 0;
-			uint8		transfer_characteristics = 0;
-			uint8		matrix_coeffs = 0;
-			uint8		chroma_loc_info_present_flag = 0;
-			uint32		chroma_sample_loc_type_top_field = 0;
-			uint32		chroma_sample_loc_type_bottom_field = 0;
-			uint8		neutral_chroma_indication_flag = 0;
-			uint8		field_seq_flag = 0;
-			uint8		frame_field_info_present_flag = 0;
-			uint8		default_display_window_flag = 0;
-			uint32		def_disp_win_left_offset = 0;
-			uint32		def_disp_win_right_offset = 0;
-			uint32		def_disp_win_top_offset = 0;
-			uint32		def_disp_win_bottom_offset = 0;
-			uint8		vui_timing_info_present_flag = 0;
-			uint32		vui_num_units_in_tick = 0;
-			uint32		vui_time_scale = 0;
-			uint8		vui_poc_proportional_to_timing_flag = 0;
-			uint32		vui_num_ticks_poc_diff_one_minus1 = 0;
-
-
-			int32 GetWidth() const
-			{
-				return (int32) pic_width_in_luma_samples;
-			}
-
-			int32 GetHeight() const
-			{
-				return (int32) pic_height_in_luma_samples;
-			}
-
-			void GetCrop(int32& Left, int32& Right, int32& Top, int32& Bottom) const
-			{
-				if (conformance_window_flag)
-				{
-					uint32 SubWidthC = 0;
-					uint32 SubHeightC = 0;
-					switch(chroma_format_idc)
-					{
-						default:
-						case 0:
-							SubWidthC = 1;
-							SubHeightC = 1;
-							break;
-						case 1:
-							SubWidthC = 2;
-							SubHeightC = 2;
-							break;
-						case 2:
-							SubWidthC = 2;
-							SubHeightC = 1;
-							break;
-						case 3:
-							SubWidthC = 1;
-							SubHeightC = 1;
-							break;
-					}
-					Left = conf_win_left_offset * SubWidthC;
-					Right = conf_win_right_offset * SubWidthC;
-					Top	= conf_win_top_offset * SubHeightC;
-					Bottom = conf_win_bottom_offset * SubHeightC;
-				}
-				else
-				{
-					Left = Right = Top = Bottom = 0;
-				}
-			}
-
-			void GetAspect(int32& SarW, int32& SarH) const
-			{
-				if (vui_parameters_present_flag && aspect_ratio_info_present_flag)
-				{
-					switch (aspect_ratio_idc)
-					{
-						default:	SarW = SarH = 0;		break;
-						case 0:		SarW = SarH = 0;		break;
-						case 1:		SarW = SarH = 1;		break;
-						case 2:		SarW = 12; SarH = 11;	break;
-						case 3:		SarW = 10; SarH = 11;	break;
-						case 4:		SarW = 16; SarH = 11;	break;
-						case 5:		SarW = 40; SarH = 33;	break;
-						case 6:		SarW = 24; SarH = 11;	break;
-						case 7:		SarW = 20; SarH = 11;	break;
-						case 8:		SarW = 32; SarH = 11;	break;
-						case 9:		SarW = 80; SarH = 33;	break;
-						case 10:	SarW = 18; SarH = 11;	break;
-						case 11:	SarW = 15; SarH = 11;	break;
-						case 12:	SarW = 64; SarH = 33;	break;
-						case 13:	SarW = 160; SarH = 99;	break;
-						case 14:	SarW = 4; SarH = 3;		break;
-						case 15:	SarW = 3; SarH = 2;		break;
-						case 16:	SarW = 2; SarH = 1;		break;
-						case 255:	SarW = sar_width; SarH = sar_height;	break;
-					}
-				}
-				else
-				{
-					SarW = SarH = 1;
-				}
-			}
-			FFractionalValue GetTiming() const
-			{
-				if (vui_parameters_present_flag && vui_timing_info_present_flag && vui_time_scale)
-				{
-					return FFractionalValue(vui_time_scale, vui_num_units_in_tick);
-				}
-				return FFractionalValue();
-			}
-
-			uint64 GetConstraintFlags() const;
-			FString GetRFC6381(const TCHAR* SampleTypePrefix) const;
-
-		};
-
-		class ELECTRADECODERS_API FHEVCDecoderConfigurationRecord
-		{
-		public:
-			void SetRawData(const void* Data, int64 Size);
-			const TArray<uint8>& GetRawData() const;
-			bool CreateFromCodecSpecificData(const TArray<uint8>& InFromCSD);
-
-			void Reset();
-
-			bool Parse();
-
-			const TArray<uint8>& GetCodecSpecificData() const;
-
-			int32 GetNumberOfSPS() const;
-			const FISO23008_2_seq_parameter_set_data& GetParsedSPS(int32 SpsIndex) const;
-
-			uint8 GetGeneralProfileSpace() const { return GeneralProfileSpace; }
-			uint8 GetGeneralTierFlag() const { return GeneralTierFlag; }
-			uint8 GetGeneralProfileIDC() const { return GeneralProfileIDC; }
-			uint32 GetGeneralProfileCompatibilityFlags() const { return GeneralProfileCompatibilityFlags; }
-			uint64 GetGeneralConstraintIndicatorFlags() const { return GeneralConstraintIndicatorFlags; }
-			uint8 GetGeneralLevelIDC() const { return GeneralLevelIDC; }
-			uint8 GetChromaFormat() const { return ChromaFormat; }
-			uint8 GetBitDepthLumaMinus8() const { return BitDepthLumaMinus8; }
-			uint8 GetBitDepthChromaMinus8() const { return BitDepthChromaMinus8; }
-
-		private:
-			struct FArray
-			{
-				uint8					Completeness = 0;
-				uint8					NALUnitType = 0;
-				TArray<TArray<uint8>>	NALUs;
-			};
-			TArray<uint8>								RawData;
-			TArray<uint8>								CodecSpecificData;
-			TArray<FArray>								Arrays;
-			uint8										ConfigurationVersion = 0;
-			uint8										GeneralProfileSpace = 0;
-			uint8										GeneralTierFlag = 0;
-			uint8										GeneralProfileIDC = 0;
-			uint32										GeneralProfileCompatibilityFlags = 0;
-			uint64										GeneralConstraintIndicatorFlags = 0;
-			uint8										GeneralLevelIDC = 0;
-			uint16										MinSpatialSegmentationIDC = 0;
-			uint8										ParallelismType = 0;
-			uint8										ChromaFormat = 0;
-			uint8										BitDepthLumaMinus8 = 0;
-			uint8										BitDepthChromaMinus8 = 0;
-			uint16										AverageFrameRate = 0;
-			uint8										ConstantFrameRate = 0;
-			uint8										NumTemporalLayers = 0;
-			uint8										TemporalIdNested = 0;
-			uint8										NALUnitLengthMinus1 = 0;
-			TArray<FISO23008_2_seq_parameter_set_data>	ParsedSPSs;
-		};
-
-
-
-		//! Parses a H.264 (ISO/IEC 14496-10) bitstream for NALUs.
-		struct FNaluInfo
-		{
-			uint64		Offset;
-			uint64		Size;
-			uint8		Type;
-			uint8		UnitLength;
-		};
-		void ELECTRADECODERS_API ParseBitstreamForNALUs(TArray<FNaluInfo>& outNALUs, const void* InBitstream, uint64 InBitstreamLength);
-
-
-		//! Parses a H.264 (ISO/IEC 14496-10) SPS NALU.
-		bool ELECTRADECODERS_API ParseH264SPS(FISO14496_10_seq_parameter_set_data& OutSPS, const void* Data, int32 Size);
-
-		//! Parses a H.265 (ISO/IEC 23008-2) SPS NALU.
-		bool ELECTRADECODERS_API ParseH265SPS(FISO23008_2_seq_parameter_set_data& OutSPS, const void* Data, int32 Size);
-
-
 		struct FSEIMessage
 		{
 			enum EPayloadType
diff --git a/Engine/Plugins/Media/ElectraPlayer/Source/ElectraPlayerRuntime/Private/Runtime/Decoder/BitstreamProcessors/VideoDecoderInputBitstreamProcessorH264.cpp b/Engine/Plugins/Media/ElectraPlayer/Source/ElectraPlayerRuntime/Private/Runtime/Decoder/BitstreamProcessors/VideoDecoderInputBitstreamProcessorH264.cpp
--- a/Engine/Plugins/Media/ElectraPlayer/Source/ElectraPlayerRuntime/Private/Runtime/Decoder/BitstreamProcessors/VideoDecoderInputBitstreamProcessorH264.cpp
+++ b/Engine/Plugins/Media/ElectraPlayer/Source/ElectraPlayerRuntime/Private/Runtime/Decoder/BitstreamProcessors/VideoDecoderInputBitstreamProcessorH264.cpp
@@ -97,15 +97,14 @@ IVideoDecoderInputBitstreamProcessor::EProcessResult FVideoDecoderInputBitstream
 		if (bDifferent)
 		{
 			SPSs.Empty();
-			TArray<ElectraDecodersUtil::MPEG::FNaluInfo> NALUs;
+			TArray<ElectraDecodersUtil::MPEG::H264::FNaluInfo> NALUs;
 			const uint8* pD = InOutAccessUnit->AUCodecData->CodecSpecificData.GetData();
-			ElectraDecodersUtil::MPEG::ParseBitstreamForNALUs(NALUs, pD, InOutAccessUnit->AUCodecData->CodecSpecificData.Num());
+			ElectraDecodersUtil::MPEG::H264::ParseBitstreamForNALUs(NALUs, pD, InOutAccessUnit->AUCodecData->CodecSpecificData.Num());
 			for(int32 i=0; i<NALUs.Num(); ++i)
 			{
 				const uint8* NALU = (const uint8*)Electra::AdvancePointer(pD, NALUs[i].Offset + NALUs[i].UnitLength);
-				uint8 nal_unit_type = *NALU & 0x1f;
 				// SPS?
-				if (nal_unit_type == 7)
+				if (NALUs[i].Type == 7)
 				{
 					bool bOk = ElectraDecodersUtil::MPEG::H264::ParseSequenceParameterSet(SPSs, NALU, NALUs[i].Size);
 					check(bOk); (void)bOk;
diff --git a/Engine/Plugins/Media/ElectraPlayer/Source/ElectraPlayerRuntime/Private/Runtime/Decoder/BitstreamProcessors/VideoDecoderInputBitstreamProcessorH265.cpp b/Engine/Plugins/Media/ElectraPlayer/Source/ElectraPlayerRuntime/Private/Runtime/Decoder/BitstreamProcessors/VideoDecoderInputBitstreamProcessorH265.cpp
--- a/Engine/Plugins/Media/ElectraPlayer/Source/ElectraPlayerRuntime/Private/Runtime/Decoder/BitstreamProcessors/VideoDecoderInputBitstreamProcessorH265.cpp
+++ b/Engine/Plugins/Media/ElectraPlayer/Source/ElectraPlayerRuntime/Private/Runtime/Decoder/BitstreamProcessors/VideoDecoderInputBitstreamProcessorH265.cpp
@@ -3,7 +3,7 @@
 #include "VideoDecoderInputBitstreamProcessorH265.h"
 
 #include "Utilities/Utilities.h"
-#include "Utils/MPEG/ElectraUtilsMPEGVideo.h"
+#include "Utils/MPEG/ElectraUtilsMPEGVideo_H265.h"
 #include "StreamAccessUnitBuffer.h"
 #include "Decoder/VideoDecoderHelpers.h"
 
@@ -29,7 +29,7 @@ private:
 	{
 	public:
 		virtual ~FCodecSpecificMessages() = default;
-		TArray<ElectraDecodersUtil::MPEG::FISO23008_2_seq_parameter_set_data> SPSs;
+		TArray<ElectraDecodersUtil::MPEG::H265::FSequenceParameterSet> SPSs;
 		TArray<ElectraDecodersUtil::MPEG::FSEIMessage> CSDPrefixSEIMessages;
 		TArray<ElectraDecodersUtil::MPEG::FSEIMessage> CSDSuffixSEIMessages;
 		TArray<ElectraDecodersUtil::MPEG::FSEIMessage> PrefixSEIMessages;
@@ -40,7 +40,7 @@ private:
 
 	TSharedPtr<const FAccessUnit::CodecData, ESPMode::ThreadSafe> PreviousCodecData;
 	TSharedPtr<const FAccessUnit::CodecData, ESPMode::ThreadSafe> CurrentCodecData;
-	TArray<ElectraDecodersUtil::MPEG::FISO23008_2_seq_parameter_set_data> SPSs;
+	TArray<ElectraDecodersUtil::MPEG::H265::FSequenceParameterSet> SPSs;
 	TArray<ElectraDecodersUtil::MPEG::FSEIMessage> PrefixSEIMessages;
 	TArray<ElectraDecodersUtil::MPEG::FSEIMessage> SuffixSEIMessages;
 
@@ -108,13 +108,13 @@ IVideoDecoderInputBitstreamProcessor::EProcessResult FVideoDecoderInputBitstream
 			SuffixSEIMessages.Empty();
 			// The CSD may contain SEI messages that apply to the stream as a whole.
 			// We need to parse the CSD to get them, if there are any.
-			TArray<ElectraDecodersUtil::MPEG::FNaluInfo>	NALUs;
+			TArray<ElectraDecodersUtil::MPEG::H265::FNaluInfo> NALUs;
 			const uint8* pD = InOutAccessUnit->AUCodecData->CodecSpecificData.GetData();
-			ElectraDecodersUtil::MPEG::ParseBitstreamForNALUs(NALUs, pD, InOutAccessUnit->AUCodecData->CodecSpecificData.Num());
+			ElectraDecodersUtil::MPEG::H265::ParseBitstreamForNALUs(NALUs, pD, InOutAccessUnit->AUCodecData->CodecSpecificData.Num());
 			for(int32 i=0; i<NALUs.Num(); ++i)
 			{
 				const uint8* NALU = (const uint8*)Electra::AdvancePointer(pD, NALUs[i].Offset + NALUs[i].UnitLength);
-				uint8 nut = *NALU >> 1;
+				const uint8 nut = NALUs[i].Type;
 				// Prefix or suffix NUT?
 				if (nut == 39 || nut == 40)
 				{
@@ -123,8 +123,8 @@ IVideoDecoderInputBitstreamProcessor::EProcessResult FVideoDecoderInputBitstream
 				// SPS nut?
 				else if (nut == 33)
 				{
-					ElectraDecodersUtil::MPEG::FISO23008_2_seq_parameter_set_data sps;
-					if (ElectraDecodersUtil::MPEG::ParseH265SPS(sps, NALU, NALUs[i].Size))
+					ElectraDecodersUtil::MPEG::H265::FSequenceParameterSet sps;
+					if (ElectraDecodersUtil::MPEG::H265::ParseSequenceParameterSet(sps, NALU, NALUs[i].Size))
 					{
 						SPSs.Emplace(MoveTemp(sps));
 					}
@@ -196,20 +196,25 @@ void FVideoDecoderInputBitstreamProcessorH265::SetPropertiesOnOutput(TSharedPtr<
 	// We only process the first SPS.
 	if (Msg->SPSs.Num() > 0)
 	{
+		const ElectraDecodersUtil::MPEG::H265::FSequenceParameterSet& sps(Msg->SPSs[0]);
+
 		// Set the bit depth and the colorimetry.
 		uint8 colour_primaries=2, transfer_characteristics=2, matrix_coeffs=2;
 		uint8 video_full_range_flag=0, video_format=5;
-		num_bits = Msg->SPSs[0].bit_depth_luma_minus8 + 8;
-		if (Msg->SPSs[0].colour_description_present_flag)
-		{
-			colour_primaries = Msg->SPSs[0].colour_primaries;
-			transfer_characteristics = Msg->SPSs[0].transfer_characteristics;
-			matrix_coeffs = Msg->SPSs[0].matrix_coeffs;
-		}
-		if (Msg->SPSs[0].video_signal_type_present_flag)
+		num_bits = sps.bit_depth_luma_minus8 + 8;
+		if (sps.vui_parameters_present_flag)
 		{
-			video_full_range_flag = Msg->SPSs[0].video_full_range_flag;
-			video_format = Msg->SPSs[0].video_format;
+			if (sps.vui_parameters.colour_description_present_flag)
+			{
+				colour_primaries = sps.vui_parameters.colour_primaries;
+				transfer_characteristics = sps.vui_parameters.transfer_characteristics;
+				matrix_coeffs = sps.vui_parameters.matrix_coeffs;
+			}
+			if (sps.vui_parameters.video_signal_type_present_flag)
+			{
+				video_full_range_flag = sps.vui_parameters.video_full_range_flag;
+				video_format = sps.vui_parameters.video_format;
+			}
 		}
 
 		Colorimetry.Update(colour_primaries, transfer_characteristics, matrix_coeffs, video_full_range_flag, video_format);
diff --git a/Engine/Plugins/Media/ElectraPlayer/Source/ElectraPlayerRuntime/Private/Runtime/Demuxer/ParserISO13818-1.cpp b/Engine/Plugins/Media/ElectraPlayer/Source/ElectraPlayerRuntime/Private/Runtime/Demuxer/ParserISO13818-1.cpp
--- a/Engine/Plugins/Media/ElectraPlayer/Source/ElectraPlayerRuntime/Private/Runtime/Demuxer/ParserISO13818-1.cpp
+++ b/Engine/Plugins/Media/ElectraPlayer/Source/ElectraPlayerRuntime/Private/Runtime/Demuxer/ParserISO13818-1.cpp
@@ -4,7 +4,6 @@
 #include "Utils/ElectraBitstreamReader.h"
 #include "Utils/MPEG/ElectraUtilsMP4.h"
 #include "Utils/MPEG/ElectraUtilsMPEGAudio.h"
-#include "Utils/MPEG/ElectraUtilsMPEGVideo.h"
 #include "Utils/MPEG/ElectraUtilsMPEGVideo_H264.h"
 #include "Utils/MPEG/ElectraUtilsMPEGVideo_H265.h"
 #include "Utilities/StringHelpers.h"
@@ -1521,17 +1520,17 @@ void FParserISO13818_1::ProcessDescriptor(ElectraDecodersUtil::FElectraBitstream
 		// HEVC_video_descriptor()
 		case 56:
 		{
-			ElectraDecodersUtil::MPEG::FISO23008_2_seq_parameter_set_data sps;
-			sps.general_profile_space = (uint8) br.GetBits(2);
-			sps.general_tier_flag = (uint8) br.GetBits(1);
-			sps.general_profile_idc = (uint8) br.GetBits(5);
-			sps.general_profile_compatibility_flag = br.GetBits(32);
-			sps.general_progressive_source_flag = (uint8) br.GetBits(1);
-			sps.general_interlaced_source_flag = (uint8) br.GetBits(1);
-			sps.general_non_packed_constraint_flag = (uint8) br.GetBits(1);
-			sps.general_frame_only_constraint_flag = (uint8) br.GetBits(1);
-			sps.general_reserved_zero_44bits = br.GetBits64(44);
-			sps.general_level_idc = (uint8) br.GetBits(8);
+			ElectraDecodersUtil::MPEG::H265::FSequenceParameterSet sps;
+			sps.profile_tier_level.general_profile_space = (uint8) br.GetBits(2);
+			sps.profile_tier_level.general_tier_flag = (uint8) br.GetBits(1);
+			sps.profile_tier_level.general_profile_idc = (uint8) br.GetBits(5);
+			sps.profile_tier_level.general_profile_compatibility_flags = br.GetBits(32);
+			sps.profile_tier_level.general_progressive_source_flag = (uint8) br.GetBits(1);
+			sps.profile_tier_level.general_interlaced_source_flag = (uint8) br.GetBits(1);
+			sps.profile_tier_level.general_non_packed_constraint_flag = (uint8) br.GetBits(1);
+			sps.profile_tier_level.general_frame_only_constraint_flag = (uint8) br.GetBits(1);
+			sps.profile_tier_level.general_constraint_indicator_flags = br.GetBits64(44);
+			sps.profile_tier_level.general_level_idc = (uint8) br.GetBits(8);
 			const uint32 temporal_layer_subset_flag = br.GetBits(1);
 			const uint32 HEVC_still_present_flag = br.GetBits(1);
 			const uint32 HEVC_24hr_picture_present_flag = br.GetBits(1);
@@ -1547,12 +1546,12 @@ void FParserISO13818_1::ProcessDescriptor(ElectraDecodersUtil::FElectraBitstream
 			}
 			OutCodecInfo.SetStreamType(EStreamType::Video);
 			OutCodecInfo.SetCodec(FStreamCodecInformation::ECodec::H265);
-			OutCodecInfo.SetProfileSpace(sps.general_profile_space);
-			OutCodecInfo.SetProfileTier(sps.general_tier_flag);
-			OutCodecInfo.SetProfile(sps.general_profile_idc);
-			OutCodecInfo.SetProfileLevel(sps.general_level_idc);
+			OutCodecInfo.SetProfileSpace(sps.profile_tier_level.general_profile_space);
+			OutCodecInfo.SetProfileTier(sps.profile_tier_level.general_tier_flag);
+			OutCodecInfo.SetProfile(sps.profile_tier_level.general_profile_idc);
+			OutCodecInfo.SetProfileLevel(sps.profile_tier_level.general_level_idc);
 			OutCodecInfo.SetProfileConstraints(sps.GetConstraintFlags());
-			OutCodecInfo.SetProfileCompatibilityFlags(sps.general_profile_compatibility_flag);
+			OutCodecInfo.SetProfileCompatibilityFlags(sps.profile_tier_level.general_profile_compatibility_flags);
 			OutCodecInfo.SetCodecSpecifierRFC6381(sps.GetRFC6381(TEXT("hvc1")));
 			break;
 		}
diff --git a/Engine/Plugins/Media/ElectraPlayer/Source/ElectraPlayerRuntime/Private/Runtime/Demuxer/ParserMKV.cpp b/Engine/Plugins/Media/ElectraPlayer/Source/ElectraPlayerRuntime/Private/Runtime/Demuxer/ParserMKV.cpp
--- a/Engine/Plugins/Media/ElectraPlayer/Source/ElectraPlayerRuntime/Private/Runtime/Demuxer/ParserMKV.cpp
+++ b/Engine/Plugins/Media/ElectraPlayer/Source/ElectraPlayerRuntime/Private/Runtime/Demuxer/ParserMKV.cpp
@@ -16,7 +16,8 @@
 #include "IElectraCodecFactory.h"
 
 // Needed to extract information from headers of AVC, HEVC and AAC files.
-#include "Utils/MPEG/ElectraUtilsMPEGVideo.h"
+#include "Utils/MPEG/ElectraUtilsMPEGVideo_H264.h"
+#include "Utils/MPEG/ElectraUtilsMPEGVideo_H265.h"
 #include "Utils/MPEG/ElectraUtilsMPEGAudio.h"
 
 DECLARE_LOG_CATEGORY_EXTERN(LogElectraMKVParser, Log, All);
@@ -2728,9 +2729,8 @@ namespace Electra
 		if (InFromTrack->GetTrackType() == EMKVTrackType::Video && Codec4CC.Equals(TEXT("avc")))
 		{
 			// The codec private data is an AVCDecoderConfigurationRecord as per ISO/IEC-14496-15
-			ElectraDecodersUtil::MPEG::FAVCDecoderConfigurationRecord CodecSpecificDataAVC;
-			CodecSpecificDataAVC.SetRawData(InFromTrack->GetCodecPrivate().GetData(), InFromTrack->GetCodecPrivate().Num());
-			if (!CodecSpecificDataAVC.Parse())
+			ElectraDecodersUtil::MPEG::H264::FAVCDecoderConfigurationRecord CodecSpecificDataAVC;
+			if (!CodecSpecificDataAVC.Parse(InFromTrack->GetCodecPrivate()))
 			{
 				return false;
 			}
@@ -2744,9 +2744,13 @@ namespace Electra
 			OutCodecInformation.SetCodecSpecificData(CodecSpecificDataAVC.GetCodecSpecificData());
 			OutCodecInformation.SetDecoderConfigRecord(InFromTrack->GetCodecPrivate());
 			OutCodecInformation.SetStreamLanguageCode(InFromTrack->GetLanguage());
-			if (CodecSpecificDataAVC.GetNumberOfSPS())
+			if (CodecSpecificDataAVC.GetSequenceParameterSets().Num())
 			{
-				const ElectraDecodersUtil::MPEG::FISO14496_10_seq_parameter_set_data& sps = CodecSpecificDataAVC.GetParsedSPS(0);
+				ElectraDecodersUtil::MPEG::H264::FSequenceParameterSet sps;
+				if (!ElectraDecodersUtil::MPEG::H264::ParseSequenceParameterSet(sps, CodecSpecificDataAVC.GetSequenceParameterSets()[0]))
+				{
+					return false;
+				}
 				int32 CropL, CropR, CropT, CropB;
 				sps.GetCrop(CropL, CropR, CropT, CropB);
 				OutCodecInformation.SetResolution(FStreamCodecInformation::FResolution(sps.GetWidth() - CropL - CropR, sps.GetHeight() - CropT - CropB));
@@ -2763,7 +2767,7 @@ namespace Electra
 				// If there is no default duration set we try to calculate it from the VUI timing of the SPS if present.
 				if (InFromTrack->GetDefaultDurationNanos() == 0)
 				{
-					ElectraDecodersUtil::MPEG::FFractionalValue fv = sps.GetTiming();
+					ElectraDecodersUtil::FFractionalValue fv = sps.GetTiming();
 					FTimeFraction fr(fv.Denom, fv.Num);
 					if (fr.IsValid())
 					{
@@ -2807,9 +2811,8 @@ namespace Electra
 		else if (InFromTrack->GetTrackType() == EMKVTrackType::Video && Codec4CC.Equals(TEXT("hevc")))
 		{
 			// The codec private data is an HEVCDecoderConfigurationRecord as per ISO/IEC-14496-15
-			ElectraDecodersUtil::MPEG::FHEVCDecoderConfigurationRecord CodecSpecificDataHEVC;
-			CodecSpecificDataHEVC.SetRawData(InFromTrack->GetCodecPrivate().GetData(), InFromTrack->GetCodecPrivate().Num());
-			if (!CodecSpecificDataHEVC.Parse())
+			ElectraDecodersUtil::MPEG::H265::FHEVCDecoderConfigurationRecord CodecSpecificDataHEVC;
+			if (!CodecSpecificDataHEVC.Parse(InFromTrack->GetCodecPrivate()))
 			{
 				return false;
 			}
@@ -2823,11 +2826,11 @@ namespace Electra
 			OutCodecInformation.SetCodecSpecificData(CodecSpecificDataHEVC.GetCodecSpecificData());
 			OutCodecInformation.SetDecoderConfigRecord(InFromTrack->GetCodecPrivate());
 			OutCodecInformation.SetStreamLanguageCode(InFromTrack->GetLanguage());
-			if (CodecSpecificDataHEVC.GetNumberOfSPS() == 0)
+			if (CodecSpecificDataHEVC.GetParsedSequenceParameterSets().IsEmpty())
 			{
 				return false;
 			}
-			const ElectraDecodersUtil::MPEG::FISO23008_2_seq_parameter_set_data& sps = CodecSpecificDataHEVC.GetParsedSPS(0);
+			const ElectraDecodersUtil::MPEG::H265::FSequenceParameterSet& sps = CodecSpecificDataHEVC.GetParsedSequenceParameterSets()[0];
 			int32 CropL, CropR, CropT, CropB;
 			sps.GetCrop(CropL, CropR, CropT, CropB);
 			OutCodecInformation.SetResolution(FStreamCodecInformation::FResolution(sps.GetWidth() - CropL - CropR, sps.GetHeight() - CropT - CropB));
@@ -2836,17 +2839,17 @@ namespace Electra
 			sps.GetAspect(ar.Width, ar.Height);
 			OutCodecInformation.SetAspectRatio(ar);
 			OutCodecInformation.SetFrameRate(sps.GetTiming().Denom ? FTimeFraction(sps.GetTiming().Num, sps.GetTiming().Denom) : FTimeFraction());
-			OutCodecInformation.SetProfileSpace(sps.general_profile_space);
-			OutCodecInformation.SetProfileTier(sps.general_tier_flag);
-			OutCodecInformation.SetProfile(sps.general_profile_idc);
-			OutCodecInformation.SetProfileLevel(sps.general_level_idc);
+			OutCodecInformation.SetProfileSpace(sps.profile_tier_level.general_profile_space);
+			OutCodecInformation.SetProfileTier(sps.profile_tier_level.general_tier_flag);
+			OutCodecInformation.SetProfile(sps.profile_tier_level.general_profile_idc);
+			OutCodecInformation.SetProfileLevel(sps.profile_tier_level.general_level_idc);
 			OutCodecInformation.SetProfileConstraints(sps.GetConstraintFlags());
-			OutCodecInformation.SetProfileCompatibilityFlags(sps.general_profile_compatibility_flag);
+			OutCodecInformation.SetProfileCompatibilityFlags(sps.profile_tier_level.general_profile_compatibility_flags);
 			OutCodecInformation.SetCodecSpecifierRFC6381(sps.GetRFC6381(TEXT("hvc1")));
 			// If there is no default duration set we try to calculate it from the VUI timing of the SPS if present.
 			if (InFromTrack->GetDefaultDurationNanos() == 0)
 			{
-				ElectraDecodersUtil::MPEG::FFractionalValue fv = sps.GetTiming();
+				ElectraDecodersUtil::FFractionalValue fv = sps.GetTiming();
 				FTimeFraction fr(fv.Denom, fv.Num);
 				if (fr.IsValid())
 				{
diff --git a/Engine/Plugins/Media/ElectraPlayer/Source/ElectraPlayerRuntime/Private/Runtime/Player/StreamSegmentReaderCommon.cpp b/Engine/Plugins/Media/ElectraPlayer/Source/ElectraPlayerRuntime/Private/Runtime/Player/StreamSegmentReaderCommon.cpp
--- a/Engine/Plugins/Media/ElectraPlayer/Source/ElectraPlayerRuntime/Private/Runtime/Player/StreamSegmentReaderCommon.cpp
+++ b/Engine/Plugins/Media/ElectraPlayer/Source/ElectraPlayerRuntime/Private/Runtime/Player/StreamSegmentReaderCommon.cpp
@@ -24,8 +24,9 @@
 #include "Utilities/UtilsMP4.h"
 #include "Utilities/UtilsMPEG.h"
 #include "Utils/Google/ElectraUtilsVPxVideo.h"
+#include "Utils/MPEG/ElectraUtilsMPEGVideo_H264.h"
+#include "Utils/MPEG/ElectraUtilsMPEGVideo_H265.h"
 #include "Utils/MPEG/ElectraUtilsMPEGAudio.h"
-#include "Utils/MPEG/ElectraUtilsMPEGVideo.h"
 #include "Utils/ElectraBitstreamReader.h"
 #include "Misc/TVariant.h"
 #include "BufferedDataReader.h"
@@ -2778,7 +2779,7 @@ FStreamSegmentReaderCommon::FStreamHandler::EHandleResult FStreamSegmentReaderCo
 						if (CurrentlyActiveTrackData->CSD->ParsedInfo.GetCodec() == FStreamCodecInformation::ECodec::H264 &&
 							CurrentlyActiveTrackData->CSD->RawCSD.IsEmpty())
 						{
-							ElectraDecodersUtil::MPEG::FAVCDecoderConfigurationRecord dcr;
+							ElectraDecodersUtil::MPEG::H264::FAVCDecoderConfigurationRecord dcr;
 							if (dcr.CreateFromCodecSpecificData(CurrentlyActiveTrackData->CSD->CodecSpecificData))
 							{
 								CurrentlyActiveTrackData->CSD->RawCSD = dcr.GetRawData();
@@ -2792,7 +2793,7 @@ FStreamSegmentReaderCommon::FStreamHandler::EHandleResult FStreamSegmentReaderCo
 						else if (CurrentlyActiveTrackData->CSD->ParsedInfo.GetCodec() == FStreamCodecInformation::ECodec::H265 &&
 							CurrentlyActiveTrackData->CSD->RawCSD.IsEmpty())
 						{
-							ElectraDecodersUtil::MPEG::FHEVCDecoderConfigurationRecord dcr;
+							ElectraDecodersUtil::MPEG::H265::FHEVCDecoderConfigurationRecord dcr;
 							if (dcr.CreateFromCodecSpecificData(CurrentlyActiveTrackData->CSD->CodecSpecificData))
 							{
 								CurrentlyActiveTrackData->CSD->RawCSD = dcr.GetRawData();
