ElectraCodecs: Extracting avc3 SPS and PPS from bitstream to enable playback of this format.

#rnx

diff --git a/Engine/Plugins/Media/ElectraCodecs/Source/ElectraDecoders/Private/Utils/MPEG/ElectraBitstreamProcessor_H264.cpp b/Engine/Plugins/Media/ElectraCodecs/Source/ElectraDecoders/Private/Utils/MPEG/ElectraBitstreamProcessor_H264.cpp
--- a/Engine/Plugins/Media/ElectraCodecs/Source/ElectraDecoders/Private/Utils/MPEG/ElectraBitstreamProcessor_H264.cpp
+++ b/Engine/Plugins/Media/ElectraCodecs/Source/ElectraDecoders/Private/Utils/MPEG/ElectraBitstreamProcessor_H264.cpp
@@ -27,9 +27,73 @@ public:
 		Error
 	};
 
-	ECSDResult ExtractSPS(const TMap<FString, FVariant>& InFromMap)
+	ECSDResult ExtractSPS(const TMap<FString, FVariant>& InFromMap, const TConstArrayView<uint8>& InInbandSPS)
 	{
-		// First try the `$avcC_box`
+		TMap<uint32, ElectraDecodersUtil::MPEG::H264::FSequenceParameterSet> NewSPSs;
+
+		auto ParseFromBitstream = [](TMap<uint32, ElectraDecodersUtil::MPEG::H264::FSequenceParameterSet>& OutNewSPSs, const TConstArrayView<uint8>& InBitstream) -> bool
+		{
+			TArray<ElectraDecodersUtil::MPEG::H264::FNaluInfo> NALUs;
+			if (!ensure(ElectraDecodersUtil::MPEG::H264::ParseBitstreamForNALUs(NALUs, InBitstream.GetData(), InBitstream.Num())))
+			{
+				return false;
+			}
+			for(int32 i=0; i<NALUs.Num(); ++i)
+			{
+				if (NALUs[i].Type != 7)
+				{
+					continue;
+				}
+				if (!ensure(ElectraDecodersUtil::MPEG::H264::ParseSequenceParameterSet(OutNewSPSs, InBitstream.GetData() + NALUs[i].Offset + NALUs[i].UnitLength, NALUs[i].Size)))
+				{
+					return false;
+				}
+			}
+			return true;
+		};
+
+		auto AreSPSsIdentical = [](const TMap<uint32, ElectraDecodersUtil::MPEG::H264::FSequenceParameterSet>& InA, const TMap<uint32, ElectraDecodersUtil::MPEG::H264::FSequenceParameterSet>& InB) -> bool
+		{
+			if (InA.Num() != InB.Num())
+			{
+				return false;
+			}
+			for(auto& aIt : InA)
+			{
+				auto pB = InB.Find(aIt.Key);
+				if (!pB || aIt.Value.ProcessedBitstream != pB->ProcessedBitstream)
+				{
+					return false;
+				}
+			}
+			return true;
+		};
+
+		// Check inband SPS first. As this is not removed from the bitstream it is always sent to the
+		// video decoder where it will most likely win over any prepended sideband CSD as it follows that.
+		// Inband SPS does not appear with every AU though.
+		if (InInbandSPS.Num())
+		{
+			if (CurrentDecoderConfiguration == InInbandSPS)
+			{
+				return ECSDResult::Unchanged;
+			}
+			if (!ParseFromBitstream(NewSPSs, InInbandSPS))
+			{
+				SPSs.Empty();
+				return ECSDResult::Error;
+			}
+			if (AreSPSsIdentical(NewSPSs, SPSs))
+			{
+				CurrentDecoderConfiguration = InInbandSPS;
+				return ECSDResult::Unchanged;
+			}
+			SPSs = MoveTemp(NewSPSs);
+			CurrentDecoderConfiguration = InInbandSPS;
+			return SPSs.IsEmpty() ? ECSDResult::NoCSD : ECSDResult::Changed;
+		}
+
+		// Try the `$avcC_box`
 		TArray<uint8> ConfigData = ElectraDecodersUtil::GetVariantValueUInt8Array(InFromMap, AvccBoxName);
 		// If empty try the decoder confiuration record, which is the same as the ConfigData box.
 		if (ConfigData.IsEmpty())
@@ -43,21 +107,40 @@ public:
 			{
 				return ECSDResult::Unchanged;
 			}
-			SPSs.Empty();
 			ElectraDecodersUtil::MPEG::H264::FAVCDecoderConfigurationRecord dcr;
-			if (!dcr.Parse(ConfigData) || dcr.GetSequenceParameterSets().IsEmpty())
+			if (!dcr.Parse(ConfigData))
 			{
+				SPSs.Empty();
 				return ECSDResult::Error;
 			}
+			// If the DCR has no SPS (because of avc3 or avc4) and we do not have any previous inband SPS we have to fail.
+			if (dcr.GetSequenceParameterSets().IsEmpty())
+			{
+				if (SPSs.IsEmpty())
+				{
+					return ECSDResult::Error;
+				}
+				else
+				{
+					return ECSDResult::Unchanged;
+				}
+			}
+
 			const TArray<TArray<uint8>>& dcrSPSs = dcr.GetSequenceParameterSets();
 			for(int32 i=0; i<dcrSPSs.Num(); ++i)
 			{
-				if (!ensure(ElectraDecodersUtil::MPEG::H264::ParseSequenceParameterSet(SPSs, dcrSPSs[i].GetData(), dcrSPSs[i].Num())))
+				if (!ensure(ElectraDecodersUtil::MPEG::H264::ParseSequenceParameterSet(NewSPSs, dcrSPSs[i].GetData(), dcrSPSs[i].Num())))
 				{
+					SPSs.Empty();
 					return ECSDResult::Error;
 				}
 			}
 			CurrentDecoderConfiguration = MoveTemp(ConfigData);
+			if (AreSPSsIdentical(NewSPSs, SPSs))
+			{
+				return ECSDResult::Unchanged;
+			}
+			SPSs = MoveTemp(NewSPSs);
 			return SPSs.IsEmpty() ? ECSDResult::NoCSD : ECSDResult::Changed;
 		}
 		// See if pre-extracted CSD is given.
@@ -70,24 +153,17 @@ public:
 				return ECSDResult::Unchanged;
 			}
 
-			SPSs.Empty();
-			TArray<ElectraDecodersUtil::MPEG::H264::FNaluInfo> NALUs;
-			if (!ensure(ElectraDecodersUtil::MPEG::H264::ParseBitstreamForNALUs(NALUs, ConfigData.GetData(), ConfigData.Num())))
+			if (!ParseFromBitstream(NewSPSs, ConfigData))
 			{
+				SPSs.Empty();
 				return ECSDResult::Error;
 			}
-			for(int32 i=0; i<NALUs.Num(); ++i)
+			if (AreSPSsIdentical(NewSPSs, SPSs))
 			{
-				if (NALUs[i].Type != 7)
-				{
-					continue;
-				}
-				if (!ensure(ElectraDecodersUtil::MPEG::H264::ParseSequenceParameterSet(SPSs, ConfigData.GetData() + NALUs[i].Offset + NALUs[i].UnitLength, NALUs[i].Size)))
-				{
-					return ECSDResult::Error;
-				}
+				return ECSDResult::Unchanged;
 			}
-			CurrentDecoderConfiguration = MoveTemp(ConfigData);
+			SPSs = MoveTemp(NewSPSs);
+			CurrentDecoderConfiguration = ConfigData;
 			return SPSs.IsEmpty() ? ECSDResult::NoCSD : ECSDResult::Changed;
 		}
 		return ECSDResult::NoCSD;
@@ -176,6 +252,8 @@ check(!"TODO");
 
 IElectraDecoderBitstreamProcessor::EProcessResult FElectraDecoderBitstreamProcessorH264::ProcessInputForDecoding(TSharedPtr<IElectraDecoderBitstreamInfo, ESPMode::ThreadSafe>& OutBSI, FElectraDecoderInputAccessUnit& InOutAccessUnit, const TMap<FString, FVariant>& InAccessUnitSidebandData)
 {
+	static const uint8 StartCode[4] = {0,0,0,1};
+
 	// Already processed?
 	if ((InOutAccessUnit.Flags & EElectraDecoderFlags::InputIsProcessed) != EElectraDecoderFlags::None)
 	{
@@ -184,25 +262,13 @@ IElectraDecoderBitstreamProcessor::EProcessResult FElectraDecoderBitstreamProces
 	// Set to processed even if we fail somewhere now.
 	InOutAccessUnit.Flags |= EElectraDecoderFlags::InputIsProcessed;
 
-	FImpl::ECSDResult CsdResult = Impl->ExtractSPS(InAccessUnitSidebandData);
-	if (CsdResult == FImpl::ECSDResult::Error)
-	{
-		Impl->LastErrorMessage = FString::Printf(TEXT("Failed to parse codec specific data"));
-		return EProcessResult::Error;
-	}
-	else if (CsdResult == FImpl::ECSDResult::Changed)
-	{
-		TUniquePtr<FImpl::FBitstreamInfo> BSI = MakeUnique<FImpl::FBitstreamInfo>();
-		BSI->SPSs = Impl->SPSs;
-		Impl->CurrentBSI = MakeShareable<IElectraDecoderBitstreamInfo>(BSI.Release());
-	}
-
 	// Assume this is not a sync sample and that it is discardable.
 	// We update the flags in the loop below to reflect the actual states.
 	InOutAccessUnit.Flags &= ~EElectraDecoderFlags::IsSyncSample;
 	InOutAccessUnit.Flags |= EElectraDecoderFlags::IsDiscardable;
 
 	TArray<ElectraDecodersUtil::MPEG::FSEIMessage> SEIMessages;
+	TArray<uint8> InbandSPS, InbandPPS;
 
 	// We now go over the data and replace the NALU sizes with startcode (unless the decoder wants to keep the sizes).
 	uint32* NALU = (uint32*)InOutAccessUnit.Data;
@@ -218,18 +284,29 @@ IElectraDecoderBitstreamProcessor::EProcessResult FElectraDecoderBitstreamProces
 		{
 			InOutAccessUnit.Flags &= ~EElectraDecoderFlags::IsDiscardable;
 		}
+		const uint8 nut = nal & 0x1f;
 		// IDR frame?
-		if ((nal & 0x1f) == 5)
+		if (nut == 5)
 		{
 			InOutAccessUnit.Flags |= EElectraDecoderFlags::IsSyncSample;
 		}
 		// SEI message(s)?
-		if ((nal & 0x1f) == 6)
+		else if (nut == 6)
 		{
 			ElectraDecodersUtil::MPEG::ExtractSEIMessages(SEIMessages, ElectraDecodersUtil::AdvancePointer(NALU, 5), naluLen-1, ElectraDecodersUtil::MPEG::ESEIStreamType::H264, false);
 		}
-
-		// NOTE: If we want to handle avc3 with inband SPS's we can extract them here!
+		// SPS?
+		else if (nut == 7)
+		{
+			InbandSPS.Append(StartCode, sizeof(StartCode));
+			InbandSPS.Append(reinterpret_cast<const uint8*>(ElectraDecodersUtil::AdvancePointer(NALU, 4)), naluLen);
+		}
+		// PPS?
+		else if (nut == 8)
+		{
+			InbandPPS.Append(StartCode, sizeof(StartCode));
+			InbandPPS.Append(reinterpret_cast<const uint8*>(ElectraDecodersUtil::AdvancePointer(NALU, 4)), naluLen);
+		}
 
 		if (Impl->bReplaceLengthWithStartcode)
 		{
@@ -237,6 +314,27 @@ IElectraDecoderBitstreamProcessor::EProcessResult FElectraDecoderBitstreamProces
 		}
 		NALU = ElectraDecodersUtil::AdvancePointer(NALU, naluLen + 4);
 	}
+
+	// As for inband SPS and PPS, we need to have both or neither. If we have only one or the other, ignore them both.
+	if (InbandSPS.IsEmpty() != InbandPPS.IsEmpty())
+	{
+		InbandSPS.Empty();
+		InbandPPS.Empty();
+	}
+
+	FImpl::ECSDResult CsdResult = Impl->ExtractSPS(InAccessUnitSidebandData, InbandSPS);
+	if (CsdResult == FImpl::ECSDResult::Error)
+	{
+		Impl->LastErrorMessage = FString::Printf(TEXT("Failed to parse codec specific data"));
+		return EProcessResult::Error;
+	}
+	else if (CsdResult == FImpl::ECSDResult::Changed)
+	{
+		TUniquePtr<FImpl::FBitstreamInfo> BSI = MakeUnique<FImpl::FBitstreamInfo>();
+		BSI->SPSs = Impl->SPSs;
+		Impl->CurrentBSI = MakeShareable<IElectraDecoderBitstreamInfo>(BSI.Release());
+	}
+
 	// Narrow down the SEI messages to those we will ultimately handle in SetPropertiesOnOutput() to
 	// avoid creating unnecessary individual bitstream info structures
 	SEIMessages.RemoveAll([](const ElectraDecodersUtil::MPEG::FSEIMessage& m)
diff --git a/Engine/Plugins/Media/ElectraCodecs/Source/ElectraDecoders/Private/Utils/MPEG/ElectraUtilsMPEGVideo_H264.cpp b/Engine/Plugins/Media/ElectraCodecs/Source/ElectraDecoders/Private/Utils/MPEG/ElectraUtilsMPEGVideo_H264.cpp
--- a/Engine/Plugins/Media/ElectraCodecs/Source/ElectraDecoders/Private/Utils/MPEG/ElectraUtilsMPEGVideo_H264.cpp
+++ b/Engine/Plugins/Media/ElectraCodecs/Source/ElectraDecoders/Private/Utils/MPEG/ElectraUtilsMPEGVideo_H264.cpp
@@ -74,7 +74,8 @@ namespace ElectraDecodersUtil
 			TUniquePtr<FRBSP> MakeRBSP(const uint8* InData, uint64 InDataSize)
 			{
 				auto Ptr = MakeUnique<FRBSP>(InDataSize);
-				Ptr->Size = EBSPtoRBSP(Ptr->RemovedAtSrc, Ptr->Data, InData, InDataSize);
+				uint64 AdjustedSize = EBSPtoRBSP(Ptr->RemovedAtSrc, Ptr->Data.GetData(), InData, InDataSize);
+				Ptr->Data.SetNum((int32)AdjustedSize);
 				return Ptr;
 			}
 
@@ -531,7 +532,7 @@ namespace ElectraDecodersUtil
 					}
 
 				TUniquePtr<FRBSP> RBSP(MakeRBSP(InBitstream, InBitstreamLenInBytes));
-				FBitstreamReader br(RBSP->Data, RBSP->Size);
+				FBitstreamReader br(RBSP->Data.GetData(), RBSP->Data.Num());
 
 				// Verify that this is in fact an SPS NALU.
 				if ((br.GetBits(8) & 0x1f) != 0x7)
@@ -773,6 +774,8 @@ namespace ElectraDecodersUtil
 						RANGE_CHECK_FAILURE(sps.max_num_reorder_frames <= sps.max_dec_frame_buffering, TEXT("max_num_reorder_frames"));
 					}
 				}
+				RBSP->TrimToBitstream(br);
+				sps.ProcessedBitstream = MoveTemp(RBSP->Data);
 				return true;
 				#undef RANGE_CHECK_FAILURE
 			}
@@ -789,7 +792,7 @@ namespace ElectraDecodersUtil
 					}
 
 				TUniquePtr<FRBSP> RBSP(MakeRBSP(InBitstream, InBitstreamLenInBytes));
-				FBitstreamReader br(RBSP->Data, RBSP->Size);
+				FBitstreamReader br(RBSP->Data.GetData(), RBSP->Data.Num());
 
 				// Verify that this is in fact a PPS NALU.
 				if ((br.GetBits(8) & 0x1f) != 0x8)
@@ -918,7 +921,7 @@ namespace ElectraDecodersUtil
 					}
 
 				OutRBSP = MakeRBSP(InBitstream, InBitstreamLenInBytes);
-				OutRBSPReader.SetData(OutRBSP->Data, OutRBSP->Size);
+				OutRBSPReader.SetData(OutRBSP->Data.GetData(), OutRBSP->Data.Num());
 				FBitstreamReader& br(OutRBSPReader);
 
 				// Verify that this is in fact a (supported) slice.
diff --git a/Engine/Plugins/Media/ElectraCodecs/Source/ElectraDecoders/Public/Utils/MPEG/ElectraUtilsMPEGVideo_H264.h b/Engine/Plugins/Media/ElectraCodecs/Source/ElectraDecoders/Public/Utils/MPEG/ElectraUtilsMPEGVideo_H264.h
--- a/Engine/Plugins/Media/ElectraCodecs/Source/ElectraDecoders/Public/Utils/MPEG/ElectraUtilsMPEGVideo_H264.h
+++ b/Engine/Plugins/Media/ElectraCodecs/Source/ElectraDecoders/Public/Utils/MPEG/ElectraUtilsMPEGVideo_H264.h
@@ -27,15 +27,22 @@ namespace ElectraDecodersUtil
 			{
 				FRBSP(uint64 InDataSize)
 				{
-					Size = InDataSize;
-					Data = (uint8*)FMemory::Malloc(InDataSize);
+					Data.SetNumUninitialized((int32)InDataSize);
 				}
-				~FRBSP()
+				void TrimToBitstream(const FElectraBitstreamReader& InBSR)
 				{
-					FMemory::Free(Data);
+					if (InBSR.IsByteAligned())
+					{
+						Data.SetNum((int32) InBSR.GetBytePosition());
+					}
+					else
+					{
+						const uint8 Mask = 0xffU << (8 - InBSR.GetBitPosition());
+						Data.SetNum((int32) InBSR.GetBytePosition() + 1);
+						Data.Last() &= Mask;
+					}
 				}
-				uint8* Data = nullptr;
-				uint64 Size = 0;
+				TArray<uint8> Data;
 				TArray<uint64> RemovedAtSrc;
 			};
 			uint64 ELECTRADECODERS_API EBSPtoRBSP(TArray<uint64>& OutRemovedInBufPos, uint8* OutBuf, const uint8* InBuf, uint64 InNumBytes);
@@ -100,6 +107,7 @@ namespace ElectraDecodersUtil
 				}
 				void Reset()
 				{
+					ProcessedBitstream.Empty();
 					// See default values in the comments of the respective members
 					profile_idc = 0;
 					constraint_set0_flag = 0;
@@ -189,6 +197,8 @@ namespace ElectraDecodersUtil
 				ELECTRADECODERS_API FString GetFormatInfo() const;
 				ELECTRADECODERS_API static FString GetFormatInfo(uint32 InProfileIDC, uint32 InLevelIDC, uint32 InConstraintFlags);
 
+				TArray<uint8> ProcessedBitstream;
+
 				uint32 profile_idc;											// u(8)
 				uint8 constraint_set0_flag;									// u(1)
 				uint8 constraint_set1_flag;									// u(1)
