ElectraCodecs: allowing AAC audio to be decoded when channel configuration is not specified

#rnx

diff --git a/Engine/Plugins/Media/ElectraCodecs/Source/ElectraDecoders/Private/Apple/aac/AAC_AudioDecoder_Apple.cpp b/Engine/Plugins/Media/ElectraCodecs/Source/ElectraDecoders/Private/Apple/aac/AAC_AudioDecoder_Apple.cpp
--- a/Engine/Plugins/Media/ElectraCodecs/Source/ElectraDecoders/Private/Apple/aac/AAC_AudioDecoder_Apple.cpp
+++ b/Engine/Plugins/Media/ElectraCodecs/Source/ElectraDecoders/Private/Apple/aac/AAC_AudioDecoder_Apple.cpp
@@ -53,6 +53,7 @@ public:
 		ci.NumberOfChannels = (int32) ElectraDecodersUtil::GetVariantValueSafeI64(InOptions, TEXT("num_channels"), 0);
 
 		if (!(ci.NumberOfChannels <= 8 && (
+			ci.ChannelConfiguration == 0 ||		// unspecified
 			ci.ChannelConfiguration == 1 ||		// Mono
 			ci.ChannelConfiguration == 2 ||		// Stereo
 			ci.ChannelConfiguration == 3 ||		// L/C/R
diff --git a/Engine/Plugins/Media/ElectraCodecs/Source/ElectraDecoders/Private/Linux/aac/AAC_AudioDecoder_Linux.cpp b/Engine/Plugins/Media/ElectraCodecs/Source/ElectraDecoders/Private/Linux/aac/AAC_AudioDecoder_Linux.cpp
--- a/Engine/Plugins/Media/ElectraCodecs/Source/ElectraDecoders/Private/Linux/aac/AAC_AudioDecoder_Linux.cpp
+++ b/Engine/Plugins/Media/ElectraCodecs/Source/ElectraDecoders/Private/Linux/aac/AAC_AudioDecoder_Linux.cpp
@@ -63,6 +63,7 @@ public:
 		ci.NumberOfChannels = (int32) ElectraDecodersUtil::GetVariantValueSafeI64(InOptions, TEXT("num_channels"), 0);
 
 		if (!(ci.NumberOfChannels <= 6 && (
+			ci.ChannelConfiguration == 0 ||		// unspecified
 			ci.ChannelConfiguration == 1 ||		// Mono
 			ci.ChannelConfiguration == 2 ||		// Stereo
 			ci.ChannelConfiguration == 3 ||		// L/C/R
