ElectraCodecs: Parsing the resolution from the VP9 uncompressed header

#rnx

diff --git a/Engine/Plugins/Media/ElectraCodecs/Source/ElectraDecoders/Private/Utils/Google/ElectraUtilsVPxVideo.cpp b/Engine/Plugins/Media/ElectraCodecs/Source/ElectraDecoders/Private/Utils/Google/ElectraUtilsVPxVideo.cpp
--- a/Engine/Plugins/Media/ElectraCodecs/Source/ElectraDecoders/Private/Utils/Google/ElectraUtilsVPxVideo.cpp
+++ b/Engine/Plugins/Media/ElectraCodecs/Source/ElectraDecoders/Private/Utils/Google/ElectraUtilsVPxVideo.cpp
@@ -90,6 +90,28 @@ namespace ElectraDecodersUtil
 				return true;
 			};
 
+			auto ReadFrameSize = [](FVP9UncompressedHeader& OutHeader, FElectraBitstreamReader& BitReader) -> bool
+			{
+				OutHeader.frame_width = BitReader.GetBits(16) + 1;
+				OutHeader.frame_height = BitReader.GetBits(16) + 1;
+				return true;
+			};
+
+			auto ReadRenderSize = [](FVP9UncompressedHeader& OutHeader, FElectraBitstreamReader& BitReader) -> bool
+			{
+				if (BitReader.GetBits(1) /*render_and_frame_size_different*/)
+				{
+					OutHeader.render_width = BitReader.GetBits(16) + 1;
+					OutHeader.render_height = BitReader.GetBits(16) + 1;
+				}
+				else
+				{
+					OutHeader.render_width = OutHeader.frame_width;
+					OutHeader.render_height = OutHeader.frame_height;
+				}
+				return true;
+			};
+
 			if (!Data || !Size)
 			{
 				return false;
@@ -130,6 +152,16 @@ namespace ElectraDecodersUtil
 				{
 					return false;
 				}
+				// Frame size
+				if (!ReadFrameSize(OutHeader, BitReader))
+				{
+					return false;
+				}
+				// Render size
+				if (!ReadRenderSize(OutHeader, BitReader))
+				{
+					return false;
+				}
 			}
 			else
 			{
@@ -155,9 +187,19 @@ namespace ElectraDecodersUtil
 						OutHeader.color_space = 1;
 						OutHeader.subsampling_x = OutHeader.subsampling_y = 1;
 					}
+					uint8 refresh_frame_flags = BitReader.GetBits(8);
+					// Frame size
+					if (!ReadFrameSize(OutHeader, BitReader))
+					{
+						return false;
+					}
+					// Render size
+					if (!ReadRenderSize(OutHeader, BitReader))
+					{
+						return false;
+					}
 				}
 			}
-
 			return true;
 		}
 
diff --git a/Engine/Plugins/Media/ElectraCodecs/Source/ElectraDecoders/Public/Utils/Google/ElectraUtilsVPxVideo.h b/Engine/Plugins/Media/ElectraCodecs/Source/ElectraDecoders/Public/Utils/Google/ElectraUtilsVPxVideo.h
--- a/Engine/Plugins/Media/ElectraCodecs/Source/ElectraDecoders/Public/Utils/Google/ElectraUtilsVPxVideo.h
+++ b/Engine/Plugins/Media/ElectraCodecs/Source/ElectraDecoders/Public/Utils/Google/ElectraUtilsVPxVideo.h
@@ -29,6 +29,11 @@ namespace ElectraDecodersUtil
 			uint8 subsampling_x = 0;
 			uint8 subsampling_y = 0;
 
+			uint32 frame_width = 0;
+			uint32 frame_height = 0;
+
+			uint32 render_width = 0;
+			uint32 render_height = 0;
 
 			enum EColorSpace
 			{
