ElectraCodecs: Unless explicitly set the H.264 second_chroma_qp_index_offset needs to be set to the same value as the chroma_qp_index_offset, otherwise we may get incorrect colors.

#rnx

diff --git a/Engine/Plugins/Media/ElectraCodecs/Source/ElectraDecoders/Private/Utils/MPEG/ElectraUtilsMPEGVideo_H264.cpp b/Engine/Plugins/Media/ElectraCodecs/Source/ElectraDecoders/Private/Utils/MPEG/ElectraUtilsMPEGVideo_H264.cpp
--- a/Engine/Plugins/Media/ElectraCodecs/Source/ElectraDecoders/Private/Utils/MPEG/ElectraUtilsMPEGVideo_H264.cpp
+++ b/Engine/Plugins/Media/ElectraCodecs/Source/ElectraDecoders/Private/Utils/MPEG/ElectraUtilsMPEGVideo_H264.cpp
@@ -887,6 +887,7 @@ namespace ElectraDecodersUtil
 				RANGE_CHECK_FAILURE(pps.pic_init_qs_minus26 >= -26 && pps.pic_init_qs_minus26 <= 25, TEXT("pic_init_qs_minus26"));
 				pps.chroma_qp_index_offset = br.se_v();
 				RANGE_CHECK_FAILURE(pps.chroma_qp_index_offset >= -12 && pps.chroma_qp_index_offset <= 12, TEXT("chroma_qp_index_offset"));
+				pps.second_chroma_qp_index_offset = pps.chroma_qp_index_offset;
 				pps.deblocking_filter_control_present_flag = br.GetBits(1);
 				pps.constrained_intra_pred_flag = br.GetBits(1);
 				pps.redundant_pic_cnt_present_flag = br.GetBits(1);
