Electra: Removal of the platform resource delegate and the wrapping plugin

#rnx

diff --git a/Engine/Plugins/Media/AppleProResDecoderElectra/Source/AppleProResDecoderElectra/Private/ProResDecoder/Mac/ElectraMediaProResDecoder.cpp b/Engine/Plugins/Media/AppleProResDecoderElectra/Source/AppleProResDecoderElectra/Private/ProResDecoder/Mac/ElectraMediaProResDecoder.cpp
--- a/Engine/Plugins/Media/AppleProResDecoderElectra/Source/AppleProResDecoderElectra/Private/ProResDecoder/Mac/ElectraMediaProResDecoder.cpp
+++ b/Engine/Plugins/Media/AppleProResDecoderElectra/Source/AppleProResDecoderElectra/Private/ProResDecoder/Mac/ElectraMediaProResDecoder.cpp
@@ -12,7 +12,6 @@
 #include "IElectraDecoderFeaturesAndOptions.h"
 #include "IElectraDecoder.h"
 #include "IElectraDecoderOutputVideo.h"
-#include "IElectraDecoderResourceDelegate.h"
 #include "ElectraDecodersUtils.h"
 #include "Utils/MPEG/ElectraUtilsMP4.h"	// to read the ProRes header, which is given in the form of an MP4 Atom
 #include "Utils/ElectraBitstreamProcessorGenericVideo.h"
@@ -42,14 +41,6 @@
 class FVideoDecoderProResElectra;
 
 
-class FDecoderDefaultVideoOutputFormatProResElectra : public IElectraDecoderDefaultVideoOutputFormat
-{
-public:
-	virtual ~FDecoderDefaultVideoOutputFormatProResElectra()
-	{ }
-};
-
-
 class FVideoDecoderOutputProResElectra : public IElectraDecoderVideoOutput, public IElectraDecoderVideoOutputImageBuffers
 {
 public:
@@ -95,8 +86,6 @@ public:
 		}
 		return nullptr;
 	}
-	IElectraDecoderVideoOutputTransferHandle* GetTransferHandle() const override
-	{ return nullptr; }
 	IElectraDecoderVideoOutput::EImageCopyResult CopyPlatformImage(IElectraDecoderVideoOutputCopyResources* InCopyResources) const override
 	{ return IElectraDecoderVideoOutput::EImageCopyResult::NotSupported; }
 
@@ -179,7 +168,7 @@ public:
 	}
 
 
-	FVideoDecoderProResElectra(const TMap<FString, FVariant>& InOptions, TSharedPtr<IElectraDecoderResourceDelegate, ESPMode::ThreadSafe> InResourceDelegate);
+	FVideoDecoderProResElectra(const TMap<FString, FVariant>& InOptions);
 	virtual ~FVideoDecoderProResElectra();
 
 	IElectraDecoder::EType GetType() const override
@@ -193,8 +182,6 @@ public:
 	IElectraDecoder::ECSDCompatibility IsCompatibleWith(const TMap<FString, FVariant>& CSDAndAdditionalOptions) override;
 	bool ResetToCleanStart() override;
 
-	TSharedPtr<IElectraDecoderDefaultOutputFormat, ESPMode::ThreadSafe> GetDefaultOutputFormatFromCSD(const TMap<FString, FVariant>& CSDAndAdditionalOptions) override;
-
 	EDecoderError DecodeAccessUnit(const FInputAccessUnit& InInputAccessUnit, const TMap<FString, FVariant>& InAdditionalOptions) override;
 	EDecoderError SendEndOfData() override;
 	EDecoderError Flush() override;
@@ -482,7 +469,6 @@ private:
 	EConvertResult ConvertDecoderOutput();
 
 	TMap<FString, FVariant> InitialCreationOptions;
-	TWeakPtr<IElectraDecoderResourceDelegate, ESPMode::ThreadSafe> ResourceDelegate;
 
 	IElectraDecoder::FError LastError;
 
@@ -594,9 +580,9 @@ public:
 		FVideoDecoderProResElectra::GetConfigurationOptions(OutOptions);
 	}
 
-	TSharedPtr<IElectraDecoder, ESPMode::ThreadSafe> CreateDecoderForFormat(const FString& InCodecFormat, const TMap<FString, FVariant>& InOptions, TSharedPtr<IElectraDecoderResourceDelegate, ESPMode::ThreadSafe> InResourceDelegate) override
+	TSharedPtr<IElectraDecoder, ESPMode::ThreadSafe> CreateDecoderForFormat(const FString& InCodecFormat, const TMap<FString, FVariant>& InOptions) override
 	{
-		return MakeShared<FVideoDecoderProResElectra, ESPMode::ThreadSafe>(InOptions, InResourceDelegate);
+		return MakeShared<FVideoDecoderProResElectra, ESPMode::ThreadSafe>(InOptions);
 	}
 
 	static TSharedPtr<FProResVideoDecoderElectraFactory, ESPMode::ThreadSafe> Self;
@@ -633,10 +619,9 @@ void FElectraMediaProResDecoder::Shutdown()
 /*********************************************************************************************************************/
 /*********************************************************************************************************************/
 
-FVideoDecoderProResElectra::FVideoDecoderProResElectra(const TMap<FString, FVariant>& InOptions, TSharedPtr<IElectraDecoderResourceDelegate, ESPMode::ThreadSafe> InResourceDelegate)
+FVideoDecoderProResElectra::FVideoDecoderProResElectra(const TMap<FString, FVariant>& InOptions)
 {
 	InitialCreationOptions = InOptions;
-	ResourceDelegate = InResourceDelegate;
 
 	AspectW = (int32)ElectraDecodersUtil::GetVariantValueSafeI64(InOptions, TEXT("aspect_w"), 0);
 	AspectH = (int32)ElectraDecodersUtil::GetVariantValueSafeI64(InOptions, TEXT("aspect_h"), 0);
@@ -712,11 +697,6 @@ void FVideoDecoderProResElectra::InternalFlushAllInputAndOutput()
 }
 
 
-TSharedPtr<IElectraDecoderDefaultOutputFormat, ESPMode::ThreadSafe> FVideoDecoderProResElectra::GetDefaultOutputFormatFromCSD(const TMap<FString, FVariant>& CSDAndAdditionalOptions)
-{
-	return nullptr;
-}
-
 IElectraDecoder::ECSDCompatibility FVideoDecoderProResElectra::IsCompatibleWith(const TMap<FString, FVariant>& CSDAndAdditionalOptions)
 {
 	// When we have no decoder yet then we are compatible because we will be creating a decoder when needed.
diff --git a/Engine/Plugins/Media/AppleProResDecoderElectra/Source/AppleProResDecoderElectra/Private/ProResDecoder/Windows/ElectraMediaProResDecoder.cpp b/Engine/Plugins/Media/AppleProResDecoderElectra/Source/AppleProResDecoderElectra/Private/ProResDecoder/Windows/ElectraMediaProResDecoder.cpp
--- a/Engine/Plugins/Media/AppleProResDecoderElectra/Source/AppleProResDecoderElectra/Private/ProResDecoder/Windows/ElectraMediaProResDecoder.cpp
+++ b/Engine/Plugins/Media/AppleProResDecoderElectra/Source/AppleProResDecoderElectra/Private/ProResDecoder/Windows/ElectraMediaProResDecoder.cpp
@@ -12,12 +12,12 @@
 #include "IElectraDecoderFeaturesAndOptions.h"
 #include "IElectraDecoder.h"
 #include "IElectraDecoderOutputVideo.h"
-#include "IElectraDecoderResourceDelegate.h"
 #include "ElectraDecodersUtils.h"
 #include "Utils/MPEG/ElectraUtilsMP4.h"	// to read the ProRes header, which is given in the form of an MP4 Atom
 #include "Utils/ElectraBitstreamProcessorGenericVideo.h"
-#include COMPILED_PLATFORM_HEADER(ElectraDecoderGPUBufferHelpers.h)
 
+#include "ElectraDecodersPlatformResources.h"
+#include COMPILED_PLATFORM_HEADER(ElectraDecoderGPUBufferHelpers.h)
 
 #include "ProResDecoder.h"
 
@@ -32,15 +32,6 @@
 class FVideoDecoderProResElectra;
 
 
-class FDecoderDefaultVideoOutputFormatProResElectra : public IElectraDecoderDefaultVideoOutputFormat
-{
-public:
-	virtual ~FDecoderDefaultVideoOutputFormatProResElectra()
-	{ }
-
-};
-
-
 class FVideoDecoderOutputProResElectra : public IElectraDecoderVideoOutput, public IElectraDecoderVideoOutputImageBuffers
 {
 public:
@@ -114,10 +105,6 @@ public:
 		}
 		return nullptr;
 	}
-	IElectraDecoderVideoOutputTransferHandle* GetTransferHandle() const override
-	{
-		return nullptr;
-	}
 	IElectraDecoderVideoOutput::EImageCopyResult CopyPlatformImage(IElectraDecoderVideoOutputCopyResources* InCopyResources) const override
 	{
 		return IElectraDecoderVideoOutput::EImageCopyResult::NotSupported;
@@ -207,7 +194,7 @@ public:
 		OutOptions.Emplace(IElectraDecoderFeature::IsAdaptive, FVariant(true));
 	}
 
-	FVideoDecoderProResElectra(const TMap<FString, FVariant>& InOptions, TSharedPtr<IElectraDecoderResourceDelegate, ESPMode::ThreadSafe> InResourceDelegate);
+	FVideoDecoderProResElectra(const TMap<FString, FVariant>& InOptions);
 
 	virtual ~FVideoDecoderProResElectra();
 
@@ -224,8 +211,6 @@ public:
 	IElectraDecoder::ECSDCompatibility IsCompatibleWith(const TMap<FString, FVariant>& CSDAndAdditionalOptions) override;
 	bool ResetToCleanStart() override;
 
-	TSharedPtr<IElectraDecoderDefaultOutputFormat, ESPMode::ThreadSafe> GetDefaultOutputFormatFromCSD(const TMap<FString, FVariant>& CSDAndAdditionalOptions) override;
-
 	EDecoderError DecodeAccessUnit(const FInputAccessUnit& InInputAccessUnit, const TMap<FString, FVariant>& InAdditionalOptions) override;
 	EDecoderError SendEndOfData() override;
 	EDecoderError Flush() override;
@@ -377,7 +362,6 @@ private:
 	uint32 BufferAllocationSize = 0;
 
 	TMap<FString, FVariant> InitialCreationOptions;
-	TWeakPtr<IElectraDecoderResourceDelegate, ESPMode::ThreadSafe> ResourceDelegate;
 
 	IElectraDecoder::FError LastError;
 
@@ -480,9 +464,9 @@ public:
 		FVideoDecoderProResElectra::GetConfigurationOptions(OutOptions);
 	}
 
-	TSharedPtr<IElectraDecoder, ESPMode::ThreadSafe> CreateDecoderForFormat(const FString& InCodecFormat, const TMap<FString, FVariant>& InOptions, TSharedPtr<IElectraDecoderResourceDelegate, ESPMode::ThreadSafe> InResourceDelegate) override
+	TSharedPtr<IElectraDecoder, ESPMode::ThreadSafe> CreateDecoderForFormat(const FString& InCodecFormat, const TMap<FString, FVariant>& InOptions) override
 	{
-		return MakeShared<FVideoDecoderProResElectra, ESPMode::ThreadSafe>(InOptions, InResourceDelegate);
+		return MakeShared<FVideoDecoderProResElectra, ESPMode::ThreadSafe>(InOptions);
 	}
 
 	static TSharedPtr<FProResVideoDecoderElectraFactory, ESPMode::ThreadSafe> Self;
@@ -517,9 +501,8 @@ void FElectraMediaProResDecoder::Shutdown()
 /*********************************************************************************************************************/
 /*********************************************************************************************************************/
 
-FVideoDecoderProResElectra::FVideoDecoderProResElectra(const TMap<FString, FVariant>& InOptions, TSharedPtr<IElectraDecoderResourceDelegate, ESPMode::ThreadSafe> InResourceDelegate)
+FVideoDecoderProResElectra::FVideoDecoderProResElectra(const TMap<FString, FVariant>& InOptions)
 	: InitialCreationOptions(InOptions)
-	, ResourceDelegate(InResourceDelegate)
 {
 	AspectW = (int32)ElectraDecodersUtil::GetVariantValueSafeI64(InOptions, TEXT("aspect_w"), 0);
 	AspectH = (int32)ElectraDecodersUtil::GetVariantValueSafeI64(InOptions, TEXT("aspect_h"), 0);
@@ -584,12 +567,6 @@ bool FVideoDecoderProResElectra::ResetToCleanStart()
 	return true;
 }
 
-TSharedPtr<IElectraDecoderDefaultOutputFormat, ESPMode::ThreadSafe> FVideoDecoderProResElectra::GetDefaultOutputFormatFromCSD(const TMap<FString, FVariant>& CSDAndAdditionalOptions)
-{
-	return nullptr;
-}
-
-
 bool FVideoDecoderProResElectra::ParseHeader(const void* InData, int32 InNumData)
 {
 	if (InNumData < 36)
@@ -731,11 +708,8 @@ IElectraDecoder::EDecoderError FVideoDecoderProResElectra::DecodeAccessUnit(cons
 		int32 PlatformDeviceVersion = 0;
 		bool bUseGPUBuffers = false;
 #if ELECTRA_MEDIAGPUBUFFER_DX12
-		if (auto PinnedResourceDelegate = ResourceDelegate.Pin())
-		{
-			PinnedResourceDelegate->GetD3DDevice(&PlatformDevice, &PlatformDeviceVersion);
-			bUseGPUBuffers = (PlatformDevice && PlatformDeviceVersion >= 12000);
-		}
+		FElectraDecodersPlatformResources::GetD3DDeviceAndVersion(&PlatformDevice, &PlatformDeviceVersion);
+		bUseGPUBuffers = (PlatformDevice && PlatformDeviceVersion >= 12000);
 #endif
 
 		TSharedPtr<FVideoDecoderOutputProResElectra, ESPMode::ThreadSafe> NewOutput = MakeShared<FVideoDecoderOutputProResElectra>();
diff --git a/Engine/Plugins/Media/AvidDNxHDDecoderElectra/Source/Private/Windows/ElectraMediaAvidDNxHDDecoder.cpp b/Engine/Plugins/Media/AvidDNxHDDecoderElectra/Source/Private/Windows/ElectraMediaAvidDNxHDDecoder.cpp
--- a/Engine/Plugins/Media/AvidDNxHDDecoderElectra/Source/Private/Windows/ElectraMediaAvidDNxHDDecoder.cpp
+++ b/Engine/Plugins/Media/AvidDNxHDDecoderElectra/Source/Private/Windows/ElectraMediaAvidDNxHDDecoder.cpp
@@ -16,9 +16,10 @@
 #include "IElectraDecoderFeaturesAndOptions.h"
 #include "IElectraDecoder.h"
 #include "IElectraDecoderOutputVideo.h"
-#include "IElectraDecoderResourceDelegate.h"
 #include "ElectraDecodersUtils.h"
 #include "Utils/ElectraBitstreamProcessorDefault.h"
+
+#include "ElectraDecodersPlatformResources.h"
 #include COMPILED_PLATFORM_HEADER(ElectraDecoderGPUBufferHelpers.h)
 
 #include "AvidDNxCodec.h"
@@ -35,15 +36,6 @@
 class FVideoDecoderAvidDNxHDElectra;
 
 
-class FDecoderDefaultVideoOutputFormatAvidDNxHDElectra : public IElectraDecoderDefaultVideoOutputFormat
-{
-public:
-	virtual ~FDecoderDefaultVideoOutputFormatAvidDNxHDElectra()
-	{ }
-
-};
-
-
 class FVideoDecoderOutputAvidDNxHDElectra : public IElectraDecoderVideoOutput, public IElectraDecoderVideoOutputImageBuffers
 {
 public:
@@ -117,10 +109,6 @@ public:
 		}
 		return nullptr;
 	}
-	IElectraDecoderVideoOutputTransferHandle* GetTransferHandle() const override
-	{
-		return nullptr;
-	}
 	IElectraDecoderVideoOutput::EImageCopyResult CopyPlatformImage(IElectraDecoderVideoOutputCopyResources* InCopyResources) const override
 	{
 		return IElectraDecoderVideoOutput::EImageCopyResult::NotSupported;
@@ -213,7 +201,7 @@ public:
 	EElectraDecoderPlatformPixelEncoding BufferEncoding = EElectraDecoderPlatformPixelEncoding::Native;
 #if ELECTRA_MEDIAGPUBUFFER_DX12
 	FElectraMediaDecoderOutputBufferPool_DX12::FOutputData GPUBuffer;
-	TSharedPtr<IElectraDecoderResourceDelegateBase::IAsyncConsecutiveTaskSync> GPUBuffer_TaskSync;
+	TSharedPtr<FElectraDecodersPlatformResources::FAsyncConsecutiveTaskSync, ESPMode::ThreadSafe> GPUBuffer_TaskSync;
 #endif
 };
 
@@ -228,7 +216,7 @@ public:
 		OutOptions.Emplace(IElectraDecoderFeature::IsAdaptive, FVariant(true));
 	}
 
-	FVideoDecoderAvidDNxHDElectra(const TMap<FString, FVariant>& InOptions, TSharedPtr<IElectraDecoderResourceDelegate, ESPMode::ThreadSafe> InResourceDelegate);
+	FVideoDecoderAvidDNxHDElectra(const TMap<FString, FVariant>& InOptions);
 
 	virtual ~FVideoDecoderAvidDNxHDElectra();
 
@@ -245,8 +233,6 @@ public:
 	IElectraDecoder::ECSDCompatibility IsCompatibleWith(const TMap<FString, FVariant>& CSDAndAdditionalOptions) override;
 	bool ResetToCleanStart() override;
 
-	TSharedPtr<IElectraDecoderDefaultOutputFormat, ESPMode::ThreadSafe> GetDefaultOutputFormatFromCSD(const TMap<FString, FVariant>& CSDAndAdditionalOptions) override;
-
 	EDecoderError DecodeAccessUnit(const FInputAccessUnit& InInputAccessUnit, const TMap<FString, FVariant>& InAdditionalOptions) override;
 	EDecoderError SendEndOfData() override;
 	EDecoderError Flush() override;
@@ -297,12 +283,10 @@ private:
 	TSharedPtr<FVideoDecoderOutputAvidDNxHDElectra, ESPMode::ThreadSafe> CurrentOutput;
 	bool bFlushPending = false;
 
-	TWeakPtr<IElectraDecoderResourceDelegate, ESPMode::ThreadSafe> ResourceDelegate;
-
 	uint32 MaxOutputBuffers;
 #if ELECTRA_MEDIAGPUBUFFER_DX12
 	mutable TSharedPtr<FElectraMediaDecoderOutputBufferPool_DX12> D3D12ResourcePool;
-	mutable TSharedPtr<IElectraDecoderResourceDelegateBase::IAsyncConsecutiveTaskSync> TaskSync;
+	mutable TSharedPtr<FElectraDecodersPlatformResources::FAsyncConsecutiveTaskSync, ESPMode::ThreadSafe> TaskSync;
 #endif
 };
 
@@ -354,9 +338,9 @@ public:
 		FVideoDecoderAvidDNxHDElectra::GetConfigurationOptions(OutOptions);
 	}
 
-	TSharedPtr<IElectraDecoder, ESPMode::ThreadSafe> CreateDecoderForFormat(const FString& InCodecFormat, const TMap<FString, FVariant>& InOptions, TSharedPtr<IElectraDecoderResourceDelegate, ESPMode::ThreadSafe> InResourceDelegate) override
+	TSharedPtr<IElectraDecoder, ESPMode::ThreadSafe> CreateDecoderForFormat(const FString& InCodecFormat, const TMap<FString, FVariant>& InOptions) override
 	{
-		return MakeShared<FVideoDecoderAvidDNxHDElectra, ESPMode::ThreadSafe>(InOptions, InResourceDelegate);
+		return MakeShared<FVideoDecoderAvidDNxHDElectra, ESPMode::ThreadSafe>(InOptions);
 	}
 
 	static TSharedPtr<FAvidDNxHDVideoDecoderElectraFactory, ESPMode::ThreadSafe> Self;
@@ -437,10 +421,8 @@ void FElectraMediaAvidDNxHDDecoder::Shutdown()
 /*********************************************************************************************************************/
 /*********************************************************************************************************************/
 
-FVideoDecoderAvidDNxHDElectra::FVideoDecoderAvidDNxHDElectra(const TMap<FString, FVariant>& InOptions, TSharedPtr<IElectraDecoderResourceDelegate, ESPMode::ThreadSafe> InResourceDelegate)
+FVideoDecoderAvidDNxHDElectra::FVideoDecoderAvidDNxHDElectra(const TMap<FString, FVariant>& InOptions)
 {
-	ResourceDelegate = InResourceDelegate;
-
 	DisplayWidth = (int32)ElectraDecodersUtil::GetVariantValueSafeI64(InOptions, TEXT("width"), 0);
 	DisplayHeight = (int32)ElectraDecodersUtil::GetVariantValueSafeI64(InOptions, TEXT("height"), 0);
 
@@ -506,11 +488,6 @@ bool FVideoDecoderAvidDNxHDElectra::ResetToCleanStart()
 	return true;
 }
 
-TSharedPtr<IElectraDecoderDefaultOutputFormat, ESPMode::ThreadSafe> FVideoDecoderAvidDNxHDElectra::GetDefaultOutputFormatFromCSD(const TMap<FString, FVariant>& CSDAndAdditionalOptions)
-{
-	return nullptr;
-}
-
 static void CopyData(TRefCountPtr<ID3D12Resource> BufferResource, uint32 BufferPitch, const uint8* TempBuffer, uint32 TempBufferPitch, uint32 Height)
 {
 	uint8* ResourceAddr = nullptr;
@@ -566,13 +543,9 @@ IElectraDecoder::EDecoderError FVideoDecoderAvidDNxHDElectra::DecodeAccessUnit(c
 		void* PlatformDevice = nullptr;
 		int32 PlatformDeviceVersion = 0;
 		bool bUseGPUBuffers = false;
-		auto PinnedResourceDelegate = ResourceDelegate.Pin();
 #if ELECTRA_MEDIAGPUBUFFER_DX12
-		if (PinnedResourceDelegate.IsValid())
-		{
-			PinnedResourceDelegate->GetD3DDevice(&PlatformDevice, &PlatformDeviceVersion);
-			bUseGPUBuffers = (PlatformDevice && PlatformDeviceVersion >= 12000);
-		}
+		FElectraDecodersPlatformResources::GetD3DDeviceAndVersion(&PlatformDevice, &PlatformDeviceVersion);
+		bUseGPUBuffers = (PlatformDevice && PlatformDeviceVersion >= 12000);
 #endif
 
 		TSharedPtr<FVideoDecoderOutputAvidDNxHDElectra, ESPMode::ThreadSafe> NewOutput = MakeShared<FVideoDecoderOutputAvidDNxHDElectra>();
@@ -750,7 +723,7 @@ IElectraDecoder::EDecoderError FVideoDecoderAvidDNxHDElectra::DecodeAccessUnit(c
 			// Create a tasksync instance so we can have our own async jobs run in consecutive order
 			if (!TaskSync.IsValid())
 			{
-				TaskSync = PinnedResourceDelegate->CreateAsyncConsecutiveTaskSync();
+				TaskSync = FElectraDecodersPlatformResources::CreateAsyncConsecutiveTaskSync();
 			}
 
 			// Request resource and fence...
@@ -799,11 +772,11 @@ IElectraDecoder::EDecoderError FVideoDecoderAvidDNxHDElectra::DecodeAccessUnit(c
 				// If we have a temp buffer, we need to copy the data first...
 				if (!TempBuffer.IsEmpty())
 				{
-					if (PinnedResourceDelegate->RunCodeAsync([BufferResource = NewOutput->GPUBuffer.Resource, Width = NewOutput->Width, Height = NewOutput->Height, ResourcePitch, BPP, TempBuffer, BufferFence = NewOutput->GPUBuffer.Fence, BufferFenceValue = NewOutput->GPUBuffer.FenceValue]()
+					if (FElectraDecodersPlatformResources::RunCodeAsync([BufferResource = NewOutput->GPUBuffer.Resource, Width = NewOutput->Width, Height = NewOutput->Height, ResourcePitch, BPP, TempBuffer, BufferFence = NewOutput->GPUBuffer.Fence, BufferFenceValue = NewOutput->GPUBuffer.FenceValue]()
 						{
 							CopyData(BufferResource, ResourcePitch, TempBuffer.GetData(), Width * BPP, Height);
 							BufferFence->Signal(BufferFenceValue);
-						}, TaskSync.Get()))
+						}, TaskSync))
 					{
 						NewOutput->GPUBuffer_TaskSync = TaskSync;
 					}
diff --git a/Engine/Plugins/Media/D3D12VideoDecodersElectra/Source/Private/VideoDecoder_D3D12.cpp b/Engine/Plugins/Media/D3D12VideoDecodersElectra/Source/Private/VideoDecoder_D3D12.cpp
--- a/Engine/Plugins/Media/D3D12VideoDecodersElectra/Source/Private/VideoDecoder_D3D12.cpp
+++ b/Engine/Plugins/Media/D3D12VideoDecodersElectra/Source/Private/VideoDecoder_D3D12.cpp
@@ -13,6 +13,7 @@
 #include "IElectraCodecFactoryModule.h"
 #include "IElectraCodecRegistry.h"
 
+#include "ElectraDecodersPlatformResources.h"
 #include "D3D12VideoDecodersElectraModule.h"
 
 // Codec specific implementations
@@ -476,7 +477,7 @@ public:
 		GetPlatformConfigurationOptions(OutOptions);
 	}
 
-	TSharedPtr<IElectraDecoder, ESPMode::ThreadSafe> CreateDecoderForFormat(const FString& InCodecFormat, const TMap<FString, FVariant>& InOptions, TSharedPtr<IElectraDecoderResourceDelegate, ESPMode::ThreadSafe> InResourceDelegate) override
+	TSharedPtr<IElectraDecoder, ESPMode::ThreadSafe> CreateDecoderForFormat(const FString& InCodecFormat, const TMap<FString, FVariant>& InOptions) override
 	{
 		if (bDoNotUseThisDecoder)
 		{
@@ -488,17 +489,10 @@ public:
 		// Do this under lock as it may be possible that the D3D device changed and we have to rebuild the codec list.
 		FScopeLock lock(&AccessLock);
 
-		// We need a resource delegate to proceed.
-		if (!InResourceDelegate.IsValid())
-		{
-			UE_LOG(LogD3D12VideoDecodersElectra, Error, TEXT("Resource delegate has not been provided."));
-			return nullptr;
-		}
-
 		// Verify that we are using the same D3D device that was used in the initial format determination.
 		TRefCountPtr<ID3D12Device> D3DDevice;
 		int32 D3DDeviceVersion = 0;
-		if (!InResourceDelegate->GetD3DDevice((void**)D3DDevice.GetInitReference(), &D3DDeviceVersion))
+		if (!FElectraDecodersPlatformResources::GetD3DDeviceAndVersion((void**)D3DDevice.GetInitReference(), &D3DDeviceVersion))
 		{
 			UE_LOG(LogD3D12VideoDecodersElectra, Error, TEXT("Could not obtain the current RHI D3D device."));
 			return nullptr;
@@ -552,12 +546,12 @@ public:
 		{
 			case FCodecFormatHelper::ECodecType::H264:
 			{
-				New = MakeShared<FD3D12VideoDecoder_H264>(*Codec, DecodeSupport, InOptions, InResourceDelegate, D3DDevice, CurrentFormats->GetVideoDevice(), CurrentFormats->GetVideoDeviceNodeIndex());
+				New = MakeShared<FD3D12VideoDecoder_H264>(*Codec, DecodeSupport, InOptions, D3DDevice, CurrentFormats->GetVideoDevice(), CurrentFormats->GetVideoDeviceNodeIndex());
 				break;
 			}
 			case FCodecFormatHelper::ECodecType::H265:
 			{
-				New = MakeShared<FD3D12VideoDecoder_H265>(*Codec, DecodeSupport, InOptions, InResourceDelegate, D3DDevice, CurrentFormats->GetVideoDevice(), CurrentFormats->GetVideoDeviceNodeIndex());
+				New = MakeShared<FD3D12VideoDecoder_H265>(*Codec, DecodeSupport, InOptions, D3DDevice, CurrentFormats->GetVideoDevice(), CurrentFormats->GetVideoDeviceNodeIndex());
 				break;
 			}
 		}
@@ -576,11 +570,10 @@ public:
 
 
 
-FD3D12VideoDecoder::FD3D12VideoDecoder(const FCodecFormatHelper::FCodecInfo& InCodecInfo, const D3D12_FEATURE_DATA_VIDEO_DECODE_SUPPORT& InDecodeSupport, const TMap<FString, FVariant>& InOptions, TSharedPtr<IElectraDecoderResourceDelegate, ESPMode::ThreadSafe> InResourceDelegate, const TRefCountPtr<ID3D12Device>& InD3D12Device, const TRefCountPtr<ID3D12VideoDevice>& InVideoDevice, uint32 InVideoDeviceNodeIndex)
+FD3D12VideoDecoder::FD3D12VideoDecoder(const FCodecFormatHelper::FCodecInfo& InCodecInfo, const D3D12_FEATURE_DATA_VIDEO_DECODE_SUPPORT& InDecodeSupport, const TMap<FString, FVariant>& InOptions, const TRefCountPtr<ID3D12Device>& InD3D12Device, const TRefCountPtr<ID3D12VideoDevice>& InVideoDevice, uint32 InVideoDeviceNodeIndex)
 	: CodecInfo(InCodecInfo)
 	, DecodeSupport(InDecodeSupport)
 	, InitialCreationOptions(InOptions)
-	, ResourceDelegate(InResourceDelegate)
 	, D3D12Device(InD3D12Device)
 	, VideoDevice(InVideoDevice)
 	, VideoDeviceNodeIndex(InVideoDeviceNodeIndex)
diff --git a/Engine/Plugins/Media/D3D12VideoDecodersElectra/Source/Private/VideoDecoder_D3D12_Common.h b/Engine/Plugins/Media/D3D12VideoDecodersElectra/Source/Private/VideoDecoder_D3D12_Common.h
--- a/Engine/Plugins/Media/D3D12VideoDecodersElectra/Source/Private/VideoDecoder_D3D12_Common.h
+++ b/Engine/Plugins/Media/D3D12VideoDecodersElectra/Source/Private/VideoDecoder_D3D12_Common.h
@@ -8,7 +8,6 @@
 #include "Containers/Queue.h"
 
 #include "IElectraDecoder.h"
-#include "IElectraDecoderResourceDelegate.h"
 #include "IElectraDecoderFeaturesAndOptions.h"
 #include "IElectraDecoderOutputVideo.h"
 
@@ -242,10 +241,6 @@ public:
 		}
 		return nullptr;
 	}
-	IElectraDecoderVideoOutputTransferHandle* GetTransferHandle() const override
-	{
-		return nullptr;
-	}
 	IElectraDecoderVideoOutput::EImageCopyResult CopyPlatformImage(IElectraDecoderVideoOutputCopyResources* InCopyResources) const override
 	{
 		return IElectraDecoderVideoOutput::EImageCopyResult::NotSupported;
@@ -340,7 +335,7 @@ class FD3D12VideoDecoder : public IElectraDecoder
 {
 public:
 	static bool D3D12VIDEODECODERSELECTRA_API CheckPlatformDecodeCapabilities(D3D12_FEATURE_DATA_VIDEO_DECODE_SUPPORT& InOutDecodeSupport, const ElectraDecodersUtil::FMimeTypeVideoCodecInfo& InCodecInfo, const TMap<FString, FVariant>& InOptions);
-	FD3D12VideoDecoder(const FCodecFormatHelper::FCodecInfo& InCodecInfo, const D3D12_FEATURE_DATA_VIDEO_DECODE_SUPPORT& InDecodeSupport, const TMap<FString, FVariant>& InOptions, TSharedPtr<IElectraDecoderResourceDelegate, ESPMode::ThreadSafe> InResourceDelegate, const TRefCountPtr<ID3D12Device>& InD3D12Device, const TRefCountPtr<ID3D12VideoDevice>& InVideoDevice, uint32 InVideoDeviceNodeIndex);
+	FD3D12VideoDecoder(const FCodecFormatHelper::FCodecInfo& InCodecInfo, const D3D12_FEATURE_DATA_VIDEO_DECODE_SUPPORT& InDecodeSupport, const TMap<FString, FVariant>& InOptions, const TRefCountPtr<ID3D12Device>& InD3D12Device, const TRefCountPtr<ID3D12VideoDevice>& InVideoDevice, uint32 InVideoDeviceNodeIndex);
 	virtual ~FD3D12VideoDecoder();
 protected:
 	EType GetType() const override
@@ -351,7 +346,6 @@ protected:
 	void Close() override;
 	ECSDCompatibility IsCompatibleWith(const TMap<FString, FVariant>& CSDAndAdditionalOptions) = 0;
 	bool ResetToCleanStart() override;
-	TSharedPtr<IElectraDecoderDefaultOutputFormat, ESPMode::ThreadSafe> GetDefaultOutputFormatFromCSD(const TMap<FString, FVariant>& CSDAndAdditionalOptions) = 0;
 	EDecoderError DecodeAccessUnit(const FInputAccessUnit& InInputAccessUnit, const TMap<FString, FVariant>& InAdditionalOptions) = 0;
 	EDecoderError SendEndOfData() = 0;
 	EDecoderError Flush() = 0;
@@ -470,7 +464,6 @@ protected:
 	FCodecFormatHelper::FCodecInfo CodecInfo;
 	D3D12_FEATURE_DATA_VIDEO_DECODE_SUPPORT DecodeSupport;
 	TMap<FString, FVariant> InitialCreationOptions;
-	TWeakPtr<IElectraDecoderResourceDelegate, ESPMode::ThreadSafe> ResourceDelegate;
 	TRefCountPtr<ID3D12Device> D3D12Device;
 	TRefCountPtr<ID3D12VideoDevice> VideoDevice;
 	uint32 VideoDeviceNodeIndex = 0;
diff --git a/Engine/Plugins/Media/D3D12VideoDecodersElectra/Source/Private/VideoDecoder_D3D12_H264.cpp b/Engine/Plugins/Media/D3D12VideoDecodersElectra/Source/Private/VideoDecoder_D3D12_H264.cpp
--- a/Engine/Plugins/Media/D3D12VideoDecodersElectra/Source/Private/VideoDecoder_D3D12_H264.cpp
+++ b/Engine/Plugins/Media/D3D12VideoDecodersElectra/Source/Private/VideoDecoder_D3D12_H264.cpp
@@ -8,8 +8,8 @@
 namespace ElectraVideoDecodersD3D12Video
 {
 
-FD3D12VideoDecoder_H264::FD3D12VideoDecoder_H264(const FCodecFormatHelper::FCodecInfo& InCodecInfo, const D3D12_FEATURE_DATA_VIDEO_DECODE_SUPPORT& InDecodeSupport, const TMap<FString, FVariant>& InOptions, TSharedPtr<IElectraDecoderResourceDelegate, ESPMode::ThreadSafe> InResourceDelegate, const TRefCountPtr<ID3D12Device>& InD3D12Device, const TRefCountPtr<ID3D12VideoDevice>& InVideoDevice, uint32 InVideoDeviceNodeIndex)
-	: FD3D12VideoDecoder(InCodecInfo, InDecodeSupport, InOptions, InResourceDelegate, InD3D12Device, InVideoDevice, InVideoDeviceNodeIndex)
+FD3D12VideoDecoder_H264::FD3D12VideoDecoder_H264(const FCodecFormatHelper::FCodecInfo& InCodecInfo, const D3D12_FEATURE_DATA_VIDEO_DECODE_SUPPORT& InDecodeSupport, const TMap<FString, FVariant>& InOptions, const TRefCountPtr<ID3D12Device>& InD3D12Device, const TRefCountPtr<ID3D12VideoDevice>& InVideoDevice, uint32 InVideoDeviceNodeIndex)
+	: FD3D12VideoDecoder(InCodecInfo, InDecodeSupport, InOptions, InD3D12Device, InVideoDevice, InVideoDeviceNodeIndex)
 {
 }
 
@@ -56,11 +56,6 @@ IElectraDecoder::ECSDCompatibility FD3D12VideoDecoder_H264::IsCompatibleWith(con
 #endif
 }
 
-TSharedPtr<IElectraDecoderDefaultOutputFormat, ESPMode::ThreadSafe> FD3D12VideoDecoder_H264::GetDefaultOutputFormatFromCSD(const TMap<FString, FVariant>& CSDAndAdditionalOptions)
-{
-	return nullptr;
-}
-
 TSharedPtr<IElectraDecoderBitstreamProcessor, ESPMode::ThreadSafe> FD3D12VideoDecoder_H264::CreateBitstreamProcessor()
 {
 	TMap<FString, FVariant> DecoderFeatures;
diff --git a/Engine/Plugins/Media/D3D12VideoDecodersElectra/Source/Private/VideoDecoder_D3D12_H264.h b/Engine/Plugins/Media/D3D12VideoDecodersElectra/Source/Private/VideoDecoder_D3D12_H264.h
--- a/Engine/Plugins/Media/D3D12VideoDecodersElectra/Source/Private/VideoDecoder_D3D12_H264.h
+++ b/Engine/Plugins/Media/D3D12VideoDecodersElectra/Source/Private/VideoDecoder_D3D12_H264.h
@@ -12,12 +12,11 @@ namespace ElectraVideoDecodersD3D12Video
 class FD3D12VideoDecoder_H264 : public FD3D12VideoDecoder
 {
 public:
-	FD3D12VideoDecoder_H264(const FCodecFormatHelper::FCodecInfo& InCodecInfo, const D3D12_FEATURE_DATA_VIDEO_DECODE_SUPPORT& InDecodeSupport, const TMap<FString, FVariant>& InOptions, TSharedPtr<IElectraDecoderResourceDelegate, ESPMode::ThreadSafe> InResourceDelegate, const TRefCountPtr<ID3D12Device>& InD3D12Device, const TRefCountPtr<ID3D12VideoDevice>& InVideoDevice, uint32 InVideoDeviceNodeIndex);
+	FD3D12VideoDecoder_H264(const FCodecFormatHelper::FCodecInfo& InCodecInfo, const D3D12_FEATURE_DATA_VIDEO_DECODE_SUPPORT& InDecodeSupport, const TMap<FString, FVariant>& InOptions, const TRefCountPtr<ID3D12Device>& InD3D12Device, const TRefCountPtr<ID3D12VideoDevice>& InVideoDevice, uint32 InVideoDeviceNodeIndex);
 	virtual ~FD3D12VideoDecoder_H264();
 
 protected:
 	ECSDCompatibility IsCompatibleWith(const TMap<FString, FVariant>& CSDAndAdditionalOptions) override;
-	TSharedPtr<IElectraDecoderDefaultOutputFormat, ESPMode::ThreadSafe> GetDefaultOutputFormatFromCSD(const TMap<FString, FVariant>& CSDAndAdditionalOptions) override;
 	TSharedPtr<IElectraDecoderBitstreamProcessor, ESPMode::ThreadSafe> CreateBitstreamProcessor() override;
 	EDecoderError DecodeAccessUnit(const FInputAccessUnit& InInputAccessUnit, const TMap<FString, FVariant>& InAdditionalOptions) override;
 	EDecoderError SendEndOfData() override;
diff --git a/Engine/Plugins/Media/D3D12VideoDecodersElectra/Source/Private/VideoDecoder_D3D12_H265.cpp b/Engine/Plugins/Media/D3D12VideoDecodersElectra/Source/Private/VideoDecoder_D3D12_H265.cpp
--- a/Engine/Plugins/Media/D3D12VideoDecodersElectra/Source/Private/VideoDecoder_D3D12_H265.cpp
+++ b/Engine/Plugins/Media/D3D12VideoDecodersElectra/Source/Private/VideoDecoder_D3D12_H265.cpp
@@ -14,8 +14,8 @@ const uint8 FD3D12VideoDecoder_H265::ScanOrderDiag8[64] = { 0, 8, 1, 16, 9, 2, 2
 											28, 21, 14, 7, 57, 50, 43, 36, 29, 22, 15, 58, 51, 44, 37, 30,
 											23, 59, 52, 45, 38, 31, 60, 53, 46, 39, 61, 54, 47, 62, 55, 63 };
 
-FD3D12VideoDecoder_H265::FD3D12VideoDecoder_H265(const FCodecFormatHelper::FCodecInfo& InCodecInfo, const D3D12_FEATURE_DATA_VIDEO_DECODE_SUPPORT& InDecodeSupport, const TMap<FString, FVariant>& InOptions, TSharedPtr<IElectraDecoderResourceDelegate, ESPMode::ThreadSafe> InResourceDelegate, const TRefCountPtr<ID3D12Device>& InD3D12Device, const TRefCountPtr<ID3D12VideoDevice>& InVideoDevice, uint32 InVideoDeviceNodeIndex)
-	: FD3D12VideoDecoder(InCodecInfo, InDecodeSupport, InOptions, InResourceDelegate, InD3D12Device, InVideoDevice, InVideoDeviceNodeIndex)
+FD3D12VideoDecoder_H265::FD3D12VideoDecoder_H265(const FCodecFormatHelper::FCodecInfo& InCodecInfo, const D3D12_FEATURE_DATA_VIDEO_DECODE_SUPPORT& InDecodeSupport, const TMap<FString, FVariant>& InOptions, const TRefCountPtr<ID3D12Device>& InD3D12Device, const TRefCountPtr<ID3D12VideoDevice>& InVideoDevice, uint32 InVideoDeviceNodeIndex)
+	: FD3D12VideoDecoder(InCodecInfo, InDecodeSupport, InOptions, InD3D12Device, InVideoDevice, InVideoDeviceNodeIndex)
 {
 }
 
@@ -67,11 +67,6 @@ IElectraDecoder::ECSDCompatibility FD3D12VideoDecoder_H265::IsCompatibleWith(con
 #endif
 }
 
-TSharedPtr<IElectraDecoderDefaultOutputFormat, ESPMode::ThreadSafe> FD3D12VideoDecoder_H265::GetDefaultOutputFormatFromCSD(const TMap<FString, FVariant>& CSDAndAdditionalOptions)
-{
-	return nullptr;
-}
-
 TSharedPtr<IElectraDecoderBitstreamProcessor, ESPMode::ThreadSafe> FD3D12VideoDecoder_H265::CreateBitstreamProcessor()
 {
 	TMap<FString, FVariant> DecoderFeatures;
diff --git a/Engine/Plugins/Media/D3D12VideoDecodersElectra/Source/Private/VideoDecoder_D3D12_H265.h b/Engine/Plugins/Media/D3D12VideoDecodersElectra/Source/Private/VideoDecoder_D3D12_H265.h
--- a/Engine/Plugins/Media/D3D12VideoDecodersElectra/Source/Private/VideoDecoder_D3D12_H265.h
+++ b/Engine/Plugins/Media/D3D12VideoDecodersElectra/Source/Private/VideoDecoder_D3D12_H265.h
@@ -12,12 +12,11 @@ namespace ElectraVideoDecodersD3D12Video
 class FD3D12VideoDecoder_H265 : public FD3D12VideoDecoder
 {
 public:
-	FD3D12VideoDecoder_H265(const FCodecFormatHelper::FCodecInfo& InCodecInfo, const D3D12_FEATURE_DATA_VIDEO_DECODE_SUPPORT& InDecodeSupport, const TMap<FString, FVariant>& InOptions, TSharedPtr<IElectraDecoderResourceDelegate, ESPMode::ThreadSafe> InResourceDelegate, const TRefCountPtr<ID3D12Device>& InD3D12Device, const TRefCountPtr<ID3D12VideoDevice>& InVideoDevice, uint32 InVideoDeviceNodeIndex);
+	FD3D12VideoDecoder_H265(const FCodecFormatHelper::FCodecInfo& InCodecInfo, const D3D12_FEATURE_DATA_VIDEO_DECODE_SUPPORT& InDecodeSupport, const TMap<FString, FVariant>& InOptions, const TRefCountPtr<ID3D12Device>& InD3D12Device, const TRefCountPtr<ID3D12VideoDevice>& InVideoDevice, uint32 InVideoDeviceNodeIndex);
 	virtual ~FD3D12VideoDecoder_H265();
 
 protected:
 	ECSDCompatibility IsCompatibleWith(const TMap<FString, FVariant>& CSDAndAdditionalOptions) override;
-	TSharedPtr<IElectraDecoderDefaultOutputFormat, ESPMode::ThreadSafe> GetDefaultOutputFormatFromCSD(const TMap<FString, FVariant>& CSDAndAdditionalOptions) override;
 	TSharedPtr<IElectraDecoderBitstreamProcessor, ESPMode::ThreadSafe> CreateBitstreamProcessor() override;
 	EDecoderError DecodeAccessUnit(const FInputAccessUnit& InInputAccessUnit, const TMap<FString, FVariant>& InAdditionalOptions) override;
 	EDecoderError SendEndOfData() override;
diff --git a/Engine/Plugins/Media/ElectraCodecs/Source/ElectraCodecFactory/Public/IElectraCodecFactory.h b/Engine/Plugins/Media/ElectraCodecs/Source/ElectraCodecFactory/Public/IElectraCodecFactory.h
--- a/Engine/Plugins/Media/ElectraCodecs/Source/ElectraCodecFactory/Public/IElectraCodecFactory.h
+++ b/Engine/Plugins/Media/ElectraCodecs/Source/ElectraCodecFactory/Public/IElectraCodecFactory.h
@@ -52,7 +52,7 @@ public:
 	/**
 	 * Called to create a decoder for the given format.
 	 */
-	virtual TSharedPtr<IElectraDecoder, ESPMode::ThreadSafe> CreateDecoderForFormat(const FString& InCodecFormat, const TMap<FString, FVariant>& InOptions, TSharedPtr<IElectraDecoderResourceDelegate, ESPMode::ThreadSafe> InResourceDelegate) = 0;
+	virtual TSharedPtr<IElectraDecoder, ESPMode::ThreadSafe> CreateDecoderForFormat(const FString& InCodecFormat, const TMap<FString, FVariant>& InOptions) = 0;
 };
 
 /*
diff --git a/Engine/Plugins/Media/ElectraCodecs/Source/ElectraDecoders/ElectraDecoders.Build.cs b/Engine/Plugins/Media/ElectraCodecs/Source/ElectraDecoders/ElectraDecoders.Build.cs
--- a/Engine/Plugins/Media/ElectraCodecs/Source/ElectraDecoders/ElectraDecoders.Build.cs
+++ b/Engine/Plugins/Media/ElectraCodecs/Source/ElectraDecoders/ElectraDecoders.Build.cs
@@ -27,11 +27,19 @@ namespace UnrealBuildTool.Rules
 			PrivateDependencyModuleNames.AddRange(
 				new string[] {
 					"Core",
-					"SignalProcessing",
+                    "RHI",
+					"MediaUtils",
+                    "ElectraSamples",
+					"ElectraBase",
                     // Common codecs
 					"dr_libs"					// MPEG audio layer 3
 				});
 
+			PrivateIncludePathModuleNames.AddRange(
+				new string[] {
+					"Media",
+				});
+
 			PublicIncludePathModuleNames.AddRange(
 				new string[] {
 					"ElectraCodecFactory"
@@ -44,6 +52,8 @@ namespace UnrealBuildTool.Rules
 
 			if (Target.Platform.IsInGroup(UnrealPlatformGroup.Windows))
 			{
+    			PrivateDependencyModuleNames.Add("SignalProcessing");
+
 				PublicDependencyModuleNames.Add("DirectX");
 				AddEngineThirdPartyPrivateStaticDependencies(Target, "DX12");
 
diff --git a/Engine/Plugins/Media/ElectraCodecs/Source/ElectraDecoders/Private/Android/aac/AAC_AudioDecoder_Android.cpp b/Engine/Plugins/Media/ElectraCodecs/Source/ElectraDecoders/Private/Android/aac/AAC_AudioDecoder_Android.cpp
--- a/Engine/Plugins/Media/ElectraCodecs/Source/ElectraDecoders/Private/Android/aac/AAC_AudioDecoder_Android.cpp
+++ b/Engine/Plugins/Media/ElectraCodecs/Source/ElectraDecoders/Private/Android/aac/AAC_AudioDecoder_Android.cpp
@@ -64,9 +64,9 @@ public:
 		IElectraAudioDecoderAAC_Android::GetConfigurationOptions(OutOptions);
 	}
 
-	TSharedPtr<IElectraDecoder, ESPMode::ThreadSafe> CreateDecoderForFormat(const FString& InCodecFormat, const TMap<FString, FVariant>& InOptions, TSharedPtr<IElectraDecoderResourceDelegate, ESPMode::ThreadSafe> InResourceDelegate) override
+	TSharedPtr<IElectraDecoder, ESPMode::ThreadSafe> CreateDecoderForFormat(const FString& InCodecFormat, const TMap<FString, FVariant>& InOptions) override
 	{
-		return IElectraAudioDecoderAAC_Android::Create(InOptions, InResourceDelegate);
+		return IElectraAudioDecoderAAC_Android::Create(InOptions);
 	}
 
 };
diff --git a/Engine/Plugins/Media/ElectraCodecs/Source/ElectraDecoders/Private/Android/aac/AudioDecoderAAC_Android.cpp b/Engine/Plugins/Media/ElectraCodecs/Source/ElectraDecoders/Private/Android/aac/AudioDecoderAAC_Android.cpp
--- a/Engine/Plugins/Media/ElectraCodecs/Source/ElectraDecoders/Private/Android/aac/AudioDecoderAAC_Android.cpp
+++ b/Engine/Plugins/Media/ElectraCodecs/Source/ElectraDecoders/Private/Android/aac/AudioDecoderAAC_Android.cpp
@@ -9,8 +9,6 @@
 #include "IElectraDecoderFeaturesAndOptions.h"
 #include "IElectraDecoderOutputAudio.h"
 
-#include "IElectraDecoderResourceDelegate.h"
-
 #include "ElectraDecodersModule.h"
 /*********************************************************************************************************************/
 
@@ -27,25 +25,6 @@
 #define DESTROY_DECODER_WHEN_FLUSHING 1
 /*********************************************************************************************************************/
 
-class FElectraDecoderDefaultAudioOutputFormatAAC_Android : public IElectraDecoderDefaultAudioOutputFormat
-{
-public:
-	virtual ~FElectraDecoderDefaultAudioOutputFormatAAC_Android()
-	{ }
-
-	int32 GetNumChannels() const override
-	{ return NumChannels; }
-	int32 GetSampleRate() const override
-	{ return SampleRate; }
-	int32 GetNumFrames() const override
-	{ return NumFrames; }
-
-	int32 NumChannels = 0;
-	int32 SampleRate = 0;
-	int32 NumFrames = 0;
-};
-
-
 class FElectraAudioDecoderOutputAAC_Android : public IElectraDecoderAudioOutput
 {
 public:
@@ -97,7 +76,7 @@ public:
 class FElectraAudioDecoderAAC_Android : public IElectraAudioDecoderAAC_Android
 {
 public:
-	FElectraAudioDecoderAAC_Android(const TMap<FString, FVariant>& InOptions, TSharedPtr<IElectraDecoderResourceDelegate, ESPMode::ThreadSafe> InResourceDelegate);
+	FElectraAudioDecoderAAC_Android(const TMap<FString, FVariant>& InOptions);
 
 	virtual ~FElectraAudioDecoderAAC_Android();
 
@@ -112,8 +91,6 @@ public:
 	IElectraDecoder::ECSDCompatibility IsCompatibleWith(const TMap<FString, FVariant>& CSDAndAdditionalOptions) override;
 	bool ResetToCleanStart() override;
 
-	TSharedPtr<IElectraDecoderDefaultOutputFormat, ESPMode::ThreadSafe> GetDefaultOutputFormatFromCSD(const TMap<FString, FVariant>& CSDAndAdditionalOptions) override;
-
 	EDecoderError DecodeAccessUnit(const FInputAccessUnit& InInputAccessUnit, const TMap<FString, FVariant>& InAdditionalOptions) override;
 	EDecoderError SendEndOfData() override;
 	EDecoderError Flush() override;
@@ -149,8 +126,6 @@ private:
 
 	bool ConvertDecoderOutput(const IElectraAACAudioDecoderAndroidJava::FOutputBufferInfo& InInfo);
 
-	TWeakPtr<IElectraDecoderResourceDelegate, ESPMode::ThreadSafe> ResourceDelegate;
-
 	IElectraDecoder::FError LastError;
 
 	TSharedPtr<ElectraDecodersUtil::MPEG::FAACDecoderConfigurationRecord, ESPMode::ThreadSafe> ConfigRecord;
@@ -194,21 +169,15 @@ void IElectraAudioDecoderAAC_Android::GetConfigurationOptions(TMap<FString, FVar
 {
 }
 
-TSharedPtr<IElectraDecoder, ESPMode::ThreadSafe> IElectraAudioDecoderAAC_Android::Create(const TMap<FString, FVariant>& InOptions, TSharedPtr<IElectraDecoderResourceDelegate, ESPMode::ThreadSafe> InResourceDelegate)
+TSharedPtr<IElectraDecoder, ESPMode::ThreadSafe> IElectraAudioDecoderAAC_Android::Create(const TMap<FString, FVariant>& InOptions)
 {
-	if (InResourceDelegate)
-	{
-		TSharedPtr<FElectraAudioDecoderAAC_Android, ESPMode::ThreadSafe> New = MakeShared<FElectraAudioDecoderAAC_Android>(InOptions, InResourceDelegate);
-		return New;
-	}
-	return nullptr;
+	TSharedPtr<FElectraAudioDecoderAAC_Android, ESPMode::ThreadSafe> New = MakeShared<FElectraAudioDecoderAAC_Android>(InOptions);
+	return New;
 }
 
 
-FElectraAudioDecoderAAC_Android::FElectraAudioDecoderAAC_Android(const TMap<FString, FVariant>& InOptions, TSharedPtr<IElectraDecoderResourceDelegate, ESPMode::ThreadSafe> InResourceDelegate)
+FElectraAudioDecoderAAC_Android::FElectraAudioDecoderAAC_Android(const TMap<FString, FVariant>& InOptions)
 {
-	ResourceDelegate = InResourceDelegate;
-
 	// If there is codec specific data set we create a config record from it now.
 	ConfigRecord = GetCSDFromOptions(InOptions);
 }
@@ -287,21 +256,6 @@ bool FElectraAudioDecoderAAC_Android::ResetToCleanStart()
 }
 
 
-TSharedPtr<IElectraDecoderDefaultOutputFormat, ESPMode::ThreadSafe> FElectraAudioDecoderAAC_Android::GetDefaultOutputFormatFromCSD(const TMap<FString, FVariant>& CSDAndAdditionalOptions)
-{
-	TSharedPtr<ElectraDecodersUtil::MPEG::FAACDecoderConfigurationRecord, ESPMode::ThreadSafe> Cfg = GetCSDFromOptions(CSDAndAdditionalOptions);
-	if (Cfg.IsValid())
-	{
-		TSharedPtr<FElectraDecoderDefaultAudioOutputFormatAAC_Android, ESPMode::ThreadSafe> Info = MakeShared<FElectraDecoderDefaultAudioOutputFormatAAC_Android, ESPMode::ThreadSafe>();
-		Info->NumChannels = Cfg->PSSignal > 0 ? 2 : NumChannelsForConfig[Cfg->ChannelConfiguration];
-		Info->SampleRate = Cfg->ExtSamplingFrequency ? Cfg->ExtSamplingFrequency : ConfigRecord->SamplingRate;
-		Info->NumFrames = Cfg->SBRSignal > 0 ? 2048 : 1024;
-		return Info;
-	}
-	return nullptr;
-}
-
-
 IElectraDecoder::EDecoderError FElectraAudioDecoderAAC_Android::DecodeAccessUnit(const FInputAccessUnit& InInputAccessUnit, const TMap<FString, FVariant>& InAdditionalOptions)
 {
 	// If already in error do nothing!
diff --git a/Engine/Plugins/Media/ElectraCodecs/Source/ElectraDecoders/Private/Android/aac/AudioDecoderAAC_Android.h b/Engine/Plugins/Media/ElectraCodecs/Source/ElectraDecoders/Private/Android/aac/AudioDecoderAAC_Android.h
--- a/Engine/Plugins/Media/ElectraCodecs/Source/ElectraDecoders/Private/Android/aac/AudioDecoderAAC_Android.h
+++ b/Engine/Plugins/Media/ElectraCodecs/Source/ElectraDecoders/Private/Android/aac/AudioDecoderAAC_Android.h
@@ -10,7 +10,7 @@ class IElectraAudioDecoderAAC_Android : public IElectraDecoder
 {
 public:
 	static void GetConfigurationOptions(TMap<FString, FVariant>& OutOptions);
-	static TSharedPtr<IElectraDecoder, ESPMode::ThreadSafe> Create(const TMap<FString, FVariant>& InOptions, TSharedPtr<IElectraDecoderResourceDelegate, ESPMode::ThreadSafe> InResourceDelegate);
+	static TSharedPtr<IElectraDecoder, ESPMode::ThreadSafe> Create(const TMap<FString, FVariant>& InOptions);
 
 	virtual ~IElectraAudioDecoderAAC_Android() = default;
 };
diff --git a/Engine/Plugins/Media/ElectraCodecs/Source/ElectraDecoders/Private/Android/h264/H264_VideoDecoder_Android.cpp b/Engine/Plugins/Media/ElectraCodecs/Source/ElectraDecoders/Private/Android/h264/H264_VideoDecoder_Android.cpp
--- a/Engine/Plugins/Media/ElectraCodecs/Source/ElectraDecoders/Private/Android/h264/H264_VideoDecoder_Android.cpp
+++ b/Engine/Plugins/Media/ElectraCodecs/Source/ElectraDecoders/Private/Android/h264/H264_VideoDecoder_Android.cpp
@@ -63,9 +63,9 @@ public:
 		IElectraVideoDecoderH264_Android::GetConfigurationOptions(OutOptions);
 	}
 
-	TSharedPtr<IElectraDecoder, ESPMode::ThreadSafe> CreateDecoderForFormat(const FString& InCodecFormat, const TMap<FString, FVariant>& InOptions, TSharedPtr<IElectraDecoderResourceDelegate, ESPMode::ThreadSafe> InResourceDelegate) override
+	TSharedPtr<IElectraDecoder, ESPMode::ThreadSafe> CreateDecoderForFormat(const FString& InCodecFormat, const TMap<FString, FVariant>& InOptions) override
 	{
-		return IElectraVideoDecoderH264_Android::Create(InOptions, InResourceDelegate);
+		return IElectraVideoDecoderH264_Android::Create(InOptions);
 	}
 
 };
@@ -75,4 +75,3 @@ TSharedPtr<IElectraCodecFactory, ESPMode::ThreadSafe> FH264VideoDecoderAndroid::
 {
 	return MakeShared<FH264VideoDecoderFactoryAndroid, ESPMode::ThreadSafe>();
 }
-
diff --git a/Engine/Plugins/Media/ElectraCodecs/Source/ElectraDecoders/Private/Android/h264/VideoDecoderH264_Android.cpp b/Engine/Plugins/Media/ElectraCodecs/Source/ElectraDecoders/Private/Android/h264/VideoDecoderH264_Android.cpp
--- a/Engine/Plugins/Media/ElectraCodecs/Source/ElectraDecoders/Private/Android/h264/VideoDecoderH264_Android.cpp
+++ b/Engine/Plugins/Media/ElectraCodecs/Source/ElectraDecoders/Private/Android/h264/VideoDecoderH264_Android.cpp
@@ -9,7 +9,7 @@
 #include "IElectraDecoderFeaturesAndOptions.h"
 #include "IElectraDecoderOutputVideo.h"
 
-#include "IElectraDecoderResourceDelegate.h"
+#include "ElectraDecodersPlatformResources.h"
 
 #include "ElectraDecodersModule.h"
 
@@ -33,21 +33,11 @@
 class FElectraVideoDecoderH264_Android;
 
 
-class FElectraDecoderDefaultVideoOutputFormatH264_Android : public IElectraDecoderDefaultVideoOutputFormat
-{
-public:
-	virtual ~FElectraDecoderDefaultVideoOutputFormatH264_Android()
-	{ }
-};
-
-
 class FElectraVideoDecoderOutputH264_Android : public IElectraDecoderVideoOutput
 {
 public:
 	virtual ~FElectraVideoDecoderOutputH264_Android()
-	{
-		ReleaseOutputBuffer();
-	}
+	{ }
 
 	FTimespan GetPTS() const override
 	{ return PTS; }
@@ -80,13 +70,10 @@ public:
 	{ OutExtraValues = ExtraValues; }
 	void* GetPlatformOutputHandle(EElectraDecoderPlatformOutputHandleType InTypeOfHandle) const override
 	{ return nullptr; }
-	IElectraDecoderVideoOutputTransferHandle* GetTransferHandle() const override
-	{ return nullptr; }
-	IElectraDecoderVideoOutput::EImageCopyResult CopyPlatformImage(IElectraDecoderVideoOutputCopyResources* InCopyResources) const override;
+	IElectraDecoderVideoOutput::EImageCopyResult CopyPlatformImage(IElectraDecoderVideoOutputCopyResources* InCopyResources) const override
+	{ return IElectraDecoderVideoOutput::EImageCopyResult::NotSupported; }
 
 public:
-	void ReleaseOutputBuffer();
-
 	FTimespan PTS;
 	uint64 UserValue = 0;
 
@@ -101,10 +88,6 @@ public:
 	int32 FrameRateD = 0;
 	int32 PixelFormat = 0;
 	TMap<FString, FVariant> ExtraValues;
-
-	TWeakPtr<FElectraVideoDecoderH264_Android, ESPMode::ThreadSafe> OwningDecoder;
-	mutable IElectraH264VideoDecoderAndroidJava::FOutputBufferInfo OwningDecoderBufferInfo;
-	mutable bool bBufferGotReferenced = false;
 };
 
 
@@ -112,7 +95,7 @@ public:
 class FElectraVideoDecoderH264_Android : public IElectraVideoDecoderH264_Android
 {
 public:
-	FElectraVideoDecoderH264_Android(const TMap<FString, FVariant>& InOptions, TSharedPtr<IElectraDecoderResourceDelegate, ESPMode::ThreadSafe> InResourceDelegate);
+	FElectraVideoDecoderH264_Android(const TMap<FString, FVariant>& InOptions);
 
 	virtual ~FElectraVideoDecoderH264_Android();
 
@@ -127,8 +110,6 @@ public:
 	IElectraDecoder::ECSDCompatibility IsCompatibleWith(const TMap<FString, FVariant>& CSDAndAdditionalOptions) override;
 	bool ResetToCleanStart() override;
 
-	TSharedPtr<IElectraDecoderDefaultOutputFormat, ESPMode::ThreadSafe> GetDefaultOutputFormatFromCSD(const TMap<FString, FVariant>& CSDAndAdditionalOptions) override;
-
 	EDecoderError DecodeAccessUnit(const FInputAccessUnit& InInputAccessUnit, const TMap<FString, FVariant>& InAdditionalOptions) override;
 	EDecoderError SendEndOfData() override;
 	EDecoderError Flush() override;
@@ -145,25 +126,8 @@ public:
 	void Suspend() override;
 	void Resume() override;
 
-	void ReleaseOutputBuffer(const IElectraH264VideoDecoderAndroidJava::FOutputBufferInfo* InBufferInfo, bool bInRender, int64 InReleaseAt);
-
 private:
-	class FBufferReleaseCallback : public IElectraDecoderResourceDelegateAndroid::IDecoderPlatformResourceAndroid::IBufferReleaseCallback
-	{
-	public:
-		FBufferReleaseCallback(TWeakPtr<FElectraVideoDecoderH264_Android, ESPMode::ThreadSafe> InOwningDecoder)
-			: OwningDecoder(MoveTemp(InOwningDecoder))
-		{ }
-		virtual ~FBufferReleaseCallback() = default;
-		void OnReleaseBuffer(int32 InBufferIndex, TOptional<int32> InBufferValidCount, TOptional<bool> InDoRender, TOptional<int64> InRenderAt) override
-		{
-			// Note: This could potentially be called from any thread but should be executed only from the decoder thread...
-			check(!"not implemented yet");
-		}
-		TWeakPtr<FElectraVideoDecoderH264_Android, ESPMode::ThreadSafe> OwningDecoder;
-	};
-
-	class FSurfaceRequestCallback : public IElectraDecoderResourceDelegateAndroid::IDecoderPlatformResourceAndroid::ISurfaceRequestCallback
+	class FSurfaceRequestCallback : public FElectraDecodersPlatformResources::IDecoderResourceInstance::ISurfaceRequestCallback
 	{
 	public:
 		virtual ~FSurfaceRequestCallback() = default;
@@ -174,10 +138,10 @@ private:
 		jobject GetSurface() const
 		{ return Surface; }
 	private:
-		void OnNewSurface(ESurfaceType InSurfaceType, void* InSurface)
+		void OnNewSurface(ESurfaceType InSurfaceType, jobject InSurface)
 		{
 			SurfaceType = InSurfaceType;
-			Surface = (jobject) InSurface;
+			Surface = InSurface;
 			FPlatformMisc::MemoryBarrier();
 			bIsSet = true;
 		}
@@ -230,14 +194,12 @@ private:
 		GotEOS
 	};
 	EConvertResult ConvertDecoderOutput(const IElectraH264VideoDecoderAndroidJava::FOutputBufferInfo& InInfo);
-	bool HavePendingOutputIndex(int32 InIndex, bool bRemoveIfExists);
 
 	IElectraDecoder::FError LastError;
 	FInitialMaxValues InitialMaxValues;
 	TMap<FString, FVariant> InitialCreationOptions;
 
-	TWeakPtr<IElectraDecoderResourceDelegate, ESPMode::ThreadSafe> ResourceDelegate;
-	IElectraDecoderResourceDelegateAndroid::IDecoderPlatformResourceAndroid* PlatformResource = nullptr;
+	FElectraDecodersPlatformResources::IDecoderResourceInstance* ResourceInstance = nullptr;
 
 	TSharedPtr<ElectraDecodersUtil::MPEG::H264::FAVCDecoderConfigurationRecord, ESPMode::ThreadSafe> ConfigRecord;
 	TSharedPtr<ElectraDecodersUtil::MPEG::H264::FSequenceParameterSet, ESPMode::ThreadSafe> CurrentSPS;
@@ -246,7 +208,6 @@ private:
 
 	TSharedPtr<IElectraH264VideoDecoderAndroidJava, ESPMode::ThreadSafe> DecoderInstance;
 	IElectraH264VideoDecoderAndroidJava::FDecoderInformation DecoderInfo;
-	TArray<IElectraH264VideoDecoderAndroidJava::FOutputBufferInfo> PendingDecoderOutputBuffers;
 	EDecodeState DecodeState = EDecodeState::Decoding;
 	int64 LastPushedPresentationTime = -1;
 	bool bDidSendEOS = false;
@@ -255,7 +216,6 @@ private:
 	std::atomic<uint32> CurrentResumeCount;
 	bool bTriggerReplay = false;
 
-	TSharedPtr<FBufferReleaseCallback, ESPMode::ThreadSafe> CurrentBufferReleaseCallback;
 	TSharedPtr<FSurfaceRequestCallback, ESPMode::ThreadSafe> CurrentSurfaceRequest;
 	uint32 NativeDecoderID = 0;
 
@@ -273,22 +233,17 @@ void IElectraVideoDecoderH264_Android::GetConfigurationOptions(TMap<FString, FVa
 	OutOptions.Emplace(IElectraDecoderFeature::SupportsDroppingOutput, FVariant(true));
 }
 
-TSharedPtr<IElectraDecoder, ESPMode::ThreadSafe> IElectraVideoDecoderH264_Android::Create(const TMap<FString, FVariant>& InOptions, TSharedPtr<IElectraDecoderResourceDelegate, ESPMode::ThreadSafe> InResourceDelegate)
+TSharedPtr<IElectraDecoder, ESPMode::ThreadSafe> IElectraVideoDecoderH264_Android::Create(const TMap<FString, FVariant>& InOptions)
 {
-	if (InResourceDelegate)
-	{
-		TSharedPtr<FElectraVideoDecoderH264_Android, ESPMode::ThreadSafe> New = MakeShared<FElectraVideoDecoderH264_Android>(InOptions, InResourceDelegate);
-		return New;
-	}
-	return nullptr;
+	TSharedPtr<FElectraVideoDecoderH264_Android, ESPMode::ThreadSafe> New = MakeShared<FElectraVideoDecoderH264_Android>(InOptions);
+	return New;
 }
 
 
-FElectraVideoDecoderH264_Android::FElectraVideoDecoderH264_Android(const TMap<FString, FVariant>& InOptions, TSharedPtr<IElectraDecoderResourceDelegate, ESPMode::ThreadSafe> InResourceDelegate)
+FElectraVideoDecoderH264_Android::FElectraVideoDecoderH264_Android(const TMap<FString, FVariant>& InOptions)
 	: InitialCreationOptions(InOptions)
 {
-	ResourceDelegate = InResourceDelegate;
-	PlatformResource = reinterpret_cast<IElectraDecoderResourceDelegateAndroid::IDecoderPlatformResourceAndroid*>(ElectraDecodersUtil::GetVariantValueSafeU64(InOptions, TEXT("platformresource"), 0));
+	ResourceInstance = FElectraDecodersPlatformResources::CreateResourceInstance(ElectraDecodersUtil::GetVariantValueSafeI64(InOptions, FElectraDecodersPlatformResources::OptionKeyOutputTexturePoolID().ToString(), 0));
 
 	NativeDecoderID = ++NextNativeDecoderID;
 	DecoderCreatedAtResumeCount = 0;
@@ -316,6 +271,8 @@ FElectraVideoDecoderH264_Android::~FElectraVideoDecoderH264_Android()
 	check(LastError.Code == ERRCODE_INTERNAL_ALREADY_CLOSED);
 	// We do it nonetheless...
 	Close();
+
+	FElectraDecodersPlatformResources::ReleaseResourceInstance(ResourceInstance);
 }
 
 
@@ -416,7 +373,6 @@ bool FElectraVideoDecoderH264_Android::ResetToCleanStart()
 	CurrentSPS.Reset();
 	CurrentOutput.Reset();
 	InDecoderInput.Empty();
-	PendingDecoderOutputBuffers.Empty();
 	DecodeState = EDecodeState::Decoding;
 	LastPushedPresentationTime = -1;
 	bDidSendEOS = false;
@@ -425,12 +381,6 @@ bool FElectraVideoDecoderH264_Android::ResetToCleanStart()
 }
 
 
-TSharedPtr<IElectraDecoderDefaultOutputFormat, ESPMode::ThreadSafe> FElectraVideoDecoderH264_Android::GetDefaultOutputFormatFromCSD(const TMap<FString, FVariant>& CSDAndAdditionalOptions)
-{
-	return nullptr;
-}
-
-
 IElectraDecoder::EDecoderError FElectraVideoDecoderH264_Android::DecodeAccessUnit(const FInputAccessUnit& InInputAccessUnit, const TMap<FString, FVariant>& InAdditionalOptions)
 {
 	// If already in error do nothing!
@@ -477,34 +427,13 @@ IElectraDecoder::EDecoderError FElectraVideoDecoderH264_Android::DecodeAccessUni
 		}
 	}
 
-	// When using a SurfaceView we need to check if it did not suddenly change.
-	bool bNeedReplayOnNewSurface = false;
-	if (PlatformResource && CurrentSurfaceRequest.IsValid() && CurrentSurfaceRequest->HasCompleted() && CurrentSurfaceRequest->GetType() == FSurfaceRequestCallback::ESurfaceType::SurfaceView)
-	{
-		void* NewSurfaceView = nullptr;
-		IElectraDecoderResourceDelegateAndroid::IDecoderPlatformResourceAndroid::ESurfaceChangeResult SurfaceChangeResult;
-		SurfaceChangeResult = PlatformResource->VerifySurfaceView(NewSurfaceView, reinterpret_cast<void*>(CurrentSurfaceRequest->GetSurface()));
-		if (SurfaceChangeResult == IElectraDecoderResourceDelegateAndroid::IDecoderPlatformResourceAndroid::ESurfaceChangeResult::NewSurface)
-		{
-			// This depends on whether or not the decoder supports setting a new surface. If it doesn't we need to replay on a new instance.
-			check(!"Not currently handled");
-			bNeedReplayOnNewSurface = true;
-		}
-		else if (SurfaceChangeResult == IElectraDecoderResourceDelegateAndroid::IDecoderPlatformResourceAndroid::ESurfaceChangeResult::Error)
-		{
-			PostError(0, TEXT("Failed to get the new output surface"), ERRCODE_INTERNAL_INVALID_OUTPUT_SURFACE);
-			return IElectraDecoder::EDecoderError::Error;
-		}
-	}
-
 	// Resuming after suspending will re-create a decoder by replaying previous data.
 	// This may not be necessary on all devices, but is a safe catch-all.
-	if (DecoderCreatedAtResumeCount != CurrentResumeCount || bNeedReplayOnNewSurface)
+	if (DecoderCreatedAtResumeCount != CurrentResumeCount)
 	{
 		bTriggerReplay = true;
 		CurrentOutput.Reset();
 		InDecoderInput.Empty();
-		PendingDecoderOutputBuffers.Empty();
 		LastPushedPresentationTime = -1;
 		bNewDecoderRequired = true;
 	}
@@ -644,7 +573,6 @@ IElectraDecoder::EDecoderError FElectraVideoDecoderH264_Android::Flush()
 	}
 	CurrentOutput.Reset();
 	InDecoderInput.Empty();
-	PendingDecoderOutputBuffers.Empty();
 	LastPushedPresentationTime = -1;
 	bDidSendEOS = false;
 	DecodeState = EDecodeState::Decoding;
@@ -838,7 +766,6 @@ bool FElectraVideoDecoderH264_Android::InternalDecoderCreate()
 	bNewDecoderRequired = false;
 
 	CurrentSurfaceRequest.Reset();
-	CurrentBufferReleaseCallback.Reset();
 
 	DecodeState = EDecodeState::CreatingDecoder;
 	return true;
@@ -870,25 +797,13 @@ bool FElectraVideoDecoderH264_Android::InternalHandleDecoderCreate()
 			{
 				return InternalFinishDecoderCreate();
 			}
-			case FSurfaceRequestCallback::ESurfaceType::SurfaceView:
-			{
-			// TODO: When we are to use an externally provided surface we need to check with every decode call if it has changed.
-			//       If so we need to update the decoder, or if this fails re-create and replay it.
-				check(!"implement this!");
-				return PostError(0, TEXT("Failed to create decoder. Surface view is not currently supported."), ERRCODE_INTERNAL_COULD_NOT_CREATE_DECODER);
-			}
 		}
 	}
 	else
 	{
-		// We need a surface, for which we need to have a platform resource instance to get it.
-		if (!PlatformResource)
-		{
-			return PostError(0, TEXT("Failed to create decoder. No platform resource handle is set."), ERRCODE_INTERNAL_COULD_NOT_CREATE_DECODER);
-		}
 		// Create a request to get a surface and request one.
 		CurrentSurfaceRequest = MakeShared<FSurfaceRequestCallback, ESPMode::ThreadSafe>();
-		PlatformResource->RequestSurface(CurrentSurfaceRequest);
+		ResourceInstance->RequestSurface(CurrentSurfaceRequest);
 		return true;
 	}
 	return true;
@@ -958,12 +873,6 @@ bool FElectraVideoDecoderH264_Android::InternalFinishDecoderCreate()
 		}
 		DecodeState = EDecodeState::Decoding;
 
-		// Create a new buffer release interface.
-		if (PlatformResource)
-		{
-			CurrentBufferReleaseCallback = MakeShared<FBufferReleaseCallback, ESPMode::ThreadSafe>(StaticCastWeakPtr<FElectraVideoDecoderH264_Android>(AsWeak()));
-			PlatformResource->SetBufferReleaseCallback(CurrentBufferReleaseCallback);
-		}
 		return true;
 	}
 	else
@@ -984,7 +893,6 @@ void FElectraVideoDecoderH264_Android::InternalDecoderDestroy()
 	DETAILLOG(LogElectraDecoders, Log, TEXT("VideoDecoderH264::InternalDecoderDestroy()"));
 	if (DecoderInstance.IsValid())
 	{
-		CurrentBufferReleaseCallback.Reset();
 		CurrentSurfaceRequest.Reset();
 		DecoderInstance->Flush();
 		DecoderInstance->Stop();
@@ -994,36 +902,6 @@ void FElectraVideoDecoderH264_Android::InternalDecoderDestroy()
 }
 
 
-void FElectraVideoDecoderH264_Android::ReleaseOutputBuffer(const IElectraH264VideoDecoderAndroidJava::FOutputBufferInfo* InBufferInfo, bool bInRender, int64 InReleaseAt)
-{
-	if (InBufferInfo && DecoderInstance.IsValid() && HavePendingOutputIndex(InBufferInfo->BufferIndex, true))
-	{
-		int32 Result = DecoderInstance->ReleaseOutputBuffer(InBufferInfo->BufferIndex, InBufferInfo->ValidCount, bInRender, InReleaseAt);
-		if (Result)
-		{
-			PostError(Result, TEXT("Failed to release decoder output buffer"), ERRCODE_INTERNAL_FAILED_TO_CONVERT_OUTPUT_SAMPLE);
-		}
-	}
-}
-
-bool FElectraVideoDecoderH264_Android::HavePendingOutputIndex(int32 InIndex, bool bRemoveIfExists)
-{
-	for(int32 i=0; i<PendingDecoderOutputBuffers.Num(); ++i)
-	{
-		if (PendingDecoderOutputBuffers[i].BufferIndex == InIndex)
-		{
-			if (bRemoveIfExists)
-			{
-				PendingDecoderOutputBuffers.RemoveAtSwap(i);
-			}
-			return true;
-		}
-	}
-	return false;
-}
-
-
-
 FElectraVideoDecoderH264_Android::EConvertResult FElectraVideoDecoderH264_Android::ConvertDecoderOutput(const IElectraH264VideoDecoderAndroidJava::FOutputBufferInfo& InInfo)
 {
 	DETAILLOG(LogElectraDecoders, Log, TEXT("VideoDecoderH264::ConvertDecoderOutput()"));
@@ -1090,48 +968,20 @@ FElectraVideoDecoderH264_Android::EConvertResult FElectraVideoDecoderH264_Androi
 		NewOutput->FrameRateD = In->SPS->num_units_in_tick * 2;
 	}
 
-	NewOutput->OwningDecoder = StaticCastWeakPtr<FElectraVideoDecoderH264_Android>(AsWeak());
-	NewOutput->OwningDecoderBufferInfo = InInfo;
-
 	NewOutput->ExtraValues.Emplace(TEXT("platform"), FVariant(TEXT("android")));
 	NewOutput->ExtraValues.Emplace(TEXT("codec"), FVariant(TEXT("avc")));
 
-	PendingDecoderOutputBuffers.Emplace(InInfo);
 	CurrentOutput = MoveTemp(NewOutput);
-	return EConvertResult::Success;
-}
 
-
-void FElectraVideoDecoderOutputH264_Android::ReleaseOutputBuffer()
-{
-	if (OwningDecoderBufferInfo.BufferIndex >= 0)
+	if (DecoderInstance.IsValid())
 	{
-		TSharedPtr<FElectraVideoDecoderH264_Android, ESPMode::ThreadSafe> Decoder = OwningDecoder.Pin();
-		if (Decoder.IsValid())
+		int32 Result = DecoderInstance->ReleaseOutputBuffer(InInfo.BufferIndex, InInfo.ValidCount, true, -1);
+		if (Result)
 		{
-			Decoder->ReleaseOutputBuffer(&OwningDecoderBufferInfo, bBufferGotReferenced, -1);
+			PostError(Result, TEXT("Failed to release decoder output buffer"), ERRCODE_INTERNAL_FAILED_TO_CONVERT_OUTPUT_SAMPLE);
+			return EConvertResult::Failure;
 		}
-		OwningDecoderBufferInfo.BufferIndex = -1;
 	}
-}
-
-IElectraDecoderVideoOutput::EImageCopyResult FElectraVideoDecoderOutputH264_Android::CopyPlatformImage(IElectraDecoderVideoOutputCopyResources* InCopyResources) const
-{
-	TSharedPtr<FElectraVideoDecoderH264_Android, ESPMode::ThreadSafe> Decoder = OwningDecoder.Pin();
-	if (Decoder.IsValid())
-	{
-		bBufferGotReferenced = true;
-
-		InCopyResources->SetBufferIndex(OwningDecoderBufferInfo.BufferIndex);
-		InCopyResources->SetValidCount(OwningDecoderBufferInfo.ValidCount);
 
-		if (InCopyResources->ShouldReleaseBufferImmediately())
-		{
-			Decoder->ReleaseOutputBuffer(&OwningDecoderBufferInfo, bBufferGotReferenced, -1);
-			OwningDecoderBufferInfo.BufferIndex = -1;
-		}
-
-		return IElectraDecoderVideoOutput::EImageCopyResult::Succeeded;
-	}
-	return IElectraDecoderVideoOutput::EImageCopyResult::Failed;
+	return EConvertResult::Success;
 }
diff --git a/Engine/Plugins/Media/ElectraCodecs/Source/ElectraDecoders/Private/Android/h264/VideoDecoderH264_Android.h b/Engine/Plugins/Media/ElectraCodecs/Source/ElectraDecoders/Private/Android/h264/VideoDecoderH264_Android.h
--- a/Engine/Plugins/Media/ElectraCodecs/Source/ElectraDecoders/Private/Android/h264/VideoDecoderH264_Android.h
+++ b/Engine/Plugins/Media/ElectraCodecs/Source/ElectraDecoders/Private/Android/h264/VideoDecoderH264_Android.h
@@ -2,15 +2,13 @@
 
 #pragma once
 
-
 #include "IElectraDecoder.h"
-#include "IElectraDecoderResourceDelegate.h"
 
 class IElectraVideoDecoderH264_Android : public IElectraDecoder
 {
 public:
 	static void GetConfigurationOptions(TMap<FString, FVariant>& OutOptions);
-	static TSharedPtr<IElectraDecoder, ESPMode::ThreadSafe> Create(const TMap<FString, FVariant>& InOptions, TSharedPtr<IElectraDecoderResourceDelegate, ESPMode::ThreadSafe> InResourceDelegate);
+	static TSharedPtr<IElectraDecoder, ESPMode::ThreadSafe> Create(const TMap<FString, FVariant>& InOptions);
 
 	virtual ~IElectraVideoDecoderH264_Android() = default;
 };
diff --git a/Engine/Plugins/Media/ElectraCodecs/Source/ElectraDecoders/Private/Android/h265/H265_VideoDecoder_Android.cpp b/Engine/Plugins/Media/ElectraCodecs/Source/ElectraDecoders/Private/Android/h265/H265_VideoDecoder_Android.cpp
--- a/Engine/Plugins/Media/ElectraCodecs/Source/ElectraDecoders/Private/Android/h265/H265_VideoDecoder_Android.cpp
+++ b/Engine/Plugins/Media/ElectraCodecs/Source/ElectraDecoders/Private/Android/h265/H265_VideoDecoder_Android.cpp
@@ -81,9 +81,9 @@ public:
 		IElectraVideoDecoderH265_Android::GetConfigurationOptions(OutOptions);
 	}
 
-	TSharedPtr<IElectraDecoder, ESPMode::ThreadSafe> CreateDecoderForFormat(const FString& InCodecFormat, const TMap<FString, FVariant>& InOptions, TSharedPtr<IElectraDecoderResourceDelegate, ESPMode::ThreadSafe> InResourceDelegate) override
+	TSharedPtr<IElectraDecoder, ESPMode::ThreadSafe> CreateDecoderForFormat(const FString& InCodecFormat, const TMap<FString, FVariant>& InOptions) override
 	{
-		return IElectraVideoDecoderH265_Android::Create(InOptions, InResourceDelegate);
+		return IElectraVideoDecoderH265_Android::Create(InOptions);
 	}
 
 };
diff --git a/Engine/Plugins/Media/ElectraCodecs/Source/ElectraDecoders/Private/Android/h265/VideoDecoderH265_Android.cpp b/Engine/Plugins/Media/ElectraCodecs/Source/ElectraDecoders/Private/Android/h265/VideoDecoderH265_Android.cpp
--- a/Engine/Plugins/Media/ElectraCodecs/Source/ElectraDecoders/Private/Android/h265/VideoDecoderH265_Android.cpp
+++ b/Engine/Plugins/Media/ElectraCodecs/Source/ElectraDecoders/Private/Android/h265/VideoDecoderH265_Android.cpp
@@ -9,7 +9,7 @@
 #include "IElectraDecoderFeaturesAndOptions.h"
 #include "IElectraDecoderOutputVideo.h"
 
-#include "IElectraDecoderResourceDelegate.h"
+#include "ElectraDecodersPlatformResources.h"
 
 #include "ElectraDecodersModule.h"
 
@@ -33,21 +33,11 @@
 class FElectraVideoDecoderH265_Android;
 
 
-class FElectraDecoderDefaultVideoOutputFormatH265_Android : public IElectraDecoderDefaultVideoOutputFormat
-{
-public:
-	virtual ~FElectraDecoderDefaultVideoOutputFormatH265_Android()
-	{ }
-};
-
-
 class FElectraVideoDecoderOutputH265_Android : public IElectraDecoderVideoOutput
 {
 public:
 	virtual ~FElectraVideoDecoderOutputH265_Android()
-	{
-		ReleaseOutputBuffer();
-	}
+	{ }
 
 	FTimespan GetPTS() const override
 	{ return PTS; }
@@ -80,13 +70,10 @@ public:
 	{ OutExtraValues = ExtraValues; }
 	void* GetPlatformOutputHandle(EElectraDecoderPlatformOutputHandleType InTypeOfHandle) const override
 	{ return nullptr; }
-	IElectraDecoderVideoOutputTransferHandle* GetTransferHandle() const override
-	{ return nullptr; }
-	IElectraDecoderVideoOutput::EImageCopyResult CopyPlatformImage(IElectraDecoderVideoOutputCopyResources* InCopyResources) const override;
+	IElectraDecoderVideoOutput::EImageCopyResult CopyPlatformImage(IElectraDecoderVideoOutputCopyResources* InCopyResources) const override
+	{ return IElectraDecoderVideoOutput::EImageCopyResult::NotSupported; }
 
 public:
-	void ReleaseOutputBuffer();
-
 	FTimespan PTS;
 	uint64 UserValue = 0;
 
@@ -101,10 +88,6 @@ public:
 	int32 FrameRateD = 0;
 	int32 PixelFormat = 0;
 	TMap<FString, FVariant> ExtraValues;
-
-	TWeakPtr<FElectraVideoDecoderH265_Android, ESPMode::ThreadSafe> OwningDecoder;
-	mutable IElectraH265VideoDecoderAndroidJava::FOutputBufferInfo OwningDecoderBufferInfo;
-	mutable bool bBufferGotReferenced = false;
 };
 
 
@@ -112,7 +95,7 @@ public:
 class FElectraVideoDecoderH265_Android : public IElectraVideoDecoderH265_Android
 {
 public:
-	FElectraVideoDecoderH265_Android(const TMap<FString, FVariant>& InOptions, TSharedPtr<IElectraDecoderResourceDelegate, ESPMode::ThreadSafe> InResourceDelegate);
+	FElectraVideoDecoderH265_Android(const TMap<FString, FVariant>& InOptions);
 
 	virtual ~FElectraVideoDecoderH265_Android();
 
@@ -127,8 +110,6 @@ public:
 	IElectraDecoder::ECSDCompatibility IsCompatibleWith(const TMap<FString, FVariant>& CSDAndAdditionalOptions) override;
 	bool ResetToCleanStart() override;
 
-	TSharedPtr<IElectraDecoderDefaultOutputFormat, ESPMode::ThreadSafe> GetDefaultOutputFormatFromCSD(const TMap<FString, FVariant>& CSDAndAdditionalOptions) override;
-
 	EDecoderError DecodeAccessUnit(const FInputAccessUnit& InInputAccessUnit, const TMap<FString, FVariant>& InAdditionalOptions) override;
 	EDecoderError SendEndOfData() override;
 	EDecoderError Flush() override;
@@ -145,25 +126,8 @@ public:
 	void Suspend() override;
 	void Resume() override;
 
-	void ReleaseOutputBuffer(const IElectraH265VideoDecoderAndroidJava::FOutputBufferInfo* InBufferInfo, bool bInRender, int64 InReleaseAt);
-
 private:
-	class FBufferReleaseCallback : public IElectraDecoderResourceDelegateAndroid::IDecoderPlatformResourceAndroid::IBufferReleaseCallback
-	{
-	public:
-		FBufferReleaseCallback(TWeakPtr<FElectraVideoDecoderH265_Android, ESPMode::ThreadSafe> InOwningDecoder)
-			: OwningDecoder(MoveTemp(InOwningDecoder))
-		{ }
-		virtual ~FBufferReleaseCallback() = default;
-		void OnReleaseBuffer(int32 InBufferIndex, TOptional<int32> InBufferValidCount, TOptional<bool> InDoRender, TOptional<int64> InRenderAt) override
-		{
-			// Note: This could potentially be called from any thread but should be executed only from the decoder thread...
-			check(!"not implemented yet");
-		}
-		TWeakPtr<FElectraVideoDecoderH265_Android, ESPMode::ThreadSafe> OwningDecoder;
-	};
-
-	class FSurfaceRequestCallback : public IElectraDecoderResourceDelegateAndroid::IDecoderPlatformResourceAndroid::ISurfaceRequestCallback
+	class FSurfaceRequestCallback : public FElectraDecodersPlatformResources::IDecoderResourceInstance::ISurfaceRequestCallback
 	{
 	public:
 		virtual ~FSurfaceRequestCallback() = default;
@@ -174,10 +138,10 @@ private:
 		jobject GetSurface() const
 		{ return Surface; }
 	private:
-		void OnNewSurface(ESurfaceType InSurfaceType, void* InSurface)
+		void OnNewSurface(ESurfaceType InSurfaceType, jobject InSurface)
 		{
 			SurfaceType = InSurfaceType;
-			Surface = (jobject) InSurface;
+			Surface = InSurface;
 			FPlatformMisc::MemoryBarrier();
 			bIsSet = true;
 		}
@@ -230,14 +194,12 @@ private:
 		GotEOS
 	};
 	EConvertResult ConvertDecoderOutput(const IElectraH265VideoDecoderAndroidJava::FOutputBufferInfo& InInfo);
-	bool HavePendingOutputIndex(int32 InIndex, bool bRemoveIfExists);
 
 	IElectraDecoder::FError LastError;
 	FInitialMaxValues InitialMaxValues;
 	TMap<FString, FVariant> InitialCreationOptions;
 
-	TWeakPtr<IElectraDecoderResourceDelegate, ESPMode::ThreadSafe> ResourceDelegate;
-	IElectraDecoderResourceDelegateAndroid::IDecoderPlatformResourceAndroid* PlatformResource = nullptr;
+	FElectraDecodersPlatformResources::IDecoderResourceInstance* ResourceInstance = nullptr;
 
 	TSharedPtr<ElectraDecodersUtil::MPEG::H265::FHEVCDecoderConfigurationRecord, ESPMode::ThreadSafe> ConfigRecord;
 	TSharedPtr<ElectraDecodersUtil::MPEG::H265::FSequenceParameterSet, ESPMode::ThreadSafe> CurrentSPS;
@@ -246,7 +208,6 @@ private:
 
 	TSharedPtr<IElectraH265VideoDecoderAndroidJava, ESPMode::ThreadSafe> DecoderInstance;
 	IElectraH265VideoDecoderAndroidJava::FDecoderInformation DecoderInfo;
-	TArray<IElectraH265VideoDecoderAndroidJava::FOutputBufferInfo> PendingDecoderOutputBuffers;
 	EDecodeState DecodeState = EDecodeState::Decoding;
 	int64 LastPushedPresentationTime = -1;
 	bool bDidSendEOS = false;
@@ -255,7 +216,6 @@ private:
 	std::atomic<uint32> CurrentResumeCount;
 	bool bTriggerReplay = false;
 
-	TSharedPtr<FBufferReleaseCallback, ESPMode::ThreadSafe> CurrentBufferReleaseCallback;
 	TSharedPtr<FSurfaceRequestCallback, ESPMode::ThreadSafe> CurrentSurfaceRequest;
 	uint32 NativeDecoderID = 0;
 
@@ -300,22 +260,17 @@ void IElectraVideoDecoderH265_Android::GetConfigurationOptions(TMap<FString, FVa
 	OutOptions.Emplace(IElectraDecoderFeature::SupportsDroppingOutput, FVariant(true));
 }
 
-TSharedPtr<IElectraDecoder, ESPMode::ThreadSafe> IElectraVideoDecoderH265_Android::Create(const TMap<FString, FVariant>& InOptions, TSharedPtr<IElectraDecoderResourceDelegate, ESPMode::ThreadSafe> InResourceDelegate)
+TSharedPtr<IElectraDecoder, ESPMode::ThreadSafe> IElectraVideoDecoderH265_Android::Create(const TMap<FString, FVariant>& InOptions)
 {
-	if (InResourceDelegate)
-	{
-		TSharedPtr<FElectraVideoDecoderH265_Android, ESPMode::ThreadSafe> New = MakeShared<FElectraVideoDecoderH265_Android>(InOptions, InResourceDelegate);
-		return New;
-	}
-	return nullptr;
+	TSharedPtr<FElectraVideoDecoderH265_Android, ESPMode::ThreadSafe> New = MakeShared<FElectraVideoDecoderH265_Android>(InOptions);
+	return New;
 }
 
 
-FElectraVideoDecoderH265_Android::FElectraVideoDecoderH265_Android(const TMap<FString, FVariant>& InOptions, TSharedPtr<IElectraDecoderResourceDelegate, ESPMode::ThreadSafe> InResourceDelegate)
+FElectraVideoDecoderH265_Android::FElectraVideoDecoderH265_Android(const TMap<FString, FVariant>& InOptions)
 	: InitialCreationOptions(InOptions)
 {
-	ResourceDelegate = InResourceDelegate;
-	PlatformResource = reinterpret_cast<IElectraDecoderResourceDelegateAndroid::IDecoderPlatformResourceAndroid*>(ElectraDecodersUtil::GetVariantValueSafeU64(InOptions, TEXT("platformresource"), 0));
+	ResourceInstance = FElectraDecodersPlatformResources::CreateResourceInstance(ElectraDecodersUtil::GetVariantValueSafeI64(InOptions, FElectraDecodersPlatformResources::OptionKeyOutputTexturePoolID().ToString(), 0));
 
 	NativeDecoderID = ++NextNativeDecoderID;
 	DecoderCreatedAtResumeCount = 0;
@@ -343,6 +298,8 @@ FElectraVideoDecoderH265_Android::~FElectraVideoDecoderH265_Android()
 	check(LastError.Code == ERRCODE_INTERNAL_ALREADY_CLOSED);
 	// We do it nonetheless...
 	Close();
+
+	FElectraDecodersPlatformResources::ReleaseResourceInstance(ResourceInstance);
 }
 
 
@@ -442,7 +399,6 @@ bool FElectraVideoDecoderH265_Android::ResetToCleanStart()
 	CurrentSPS.Reset();
 	CurrentOutput.Reset();
 	InDecoderInput.Empty();
-	PendingDecoderOutputBuffers.Empty();
 	DecodeState = EDecodeState::Decoding;
 	LastPushedPresentationTime = -1;
 	bDidSendEOS = false;
@@ -451,12 +407,6 @@ bool FElectraVideoDecoderH265_Android::ResetToCleanStart()
 }
 
 
-TSharedPtr<IElectraDecoderDefaultOutputFormat, ESPMode::ThreadSafe> FElectraVideoDecoderH265_Android::GetDefaultOutputFormatFromCSD(const TMap<FString, FVariant>& CSDAndAdditionalOptions)
-{
-	return nullptr;
-}
-
-
 IElectraDecoder::EDecoderError FElectraVideoDecoderH265_Android::DecodeAccessUnit(const FInputAccessUnit& InInputAccessUnit, const TMap<FString, FVariant>& InAdditionalOptions)
 {
 	// If already in error do nothing!
@@ -503,34 +453,13 @@ IElectraDecoder::EDecoderError FElectraVideoDecoderH265_Android::DecodeAccessUni
 		}
 	}
 
-	// When using a SurfaceView we need to check if it did not suddenly change.
-	bool bNeedReplayOnNewSurface = false;
-	if (PlatformResource && CurrentSurfaceRequest.IsValid() && CurrentSurfaceRequest->HasCompleted() && CurrentSurfaceRequest->GetType() == FSurfaceRequestCallback::ESurfaceType::SurfaceView)
-	{
-		void* NewSurfaceView = nullptr;
-		IElectraDecoderResourceDelegateAndroid::IDecoderPlatformResourceAndroid::ESurfaceChangeResult SurfaceChangeResult;
-		SurfaceChangeResult = PlatformResource->VerifySurfaceView(NewSurfaceView, reinterpret_cast<void*>(CurrentSurfaceRequest->GetSurface()));
-		if (SurfaceChangeResult == IElectraDecoderResourceDelegateAndroid::IDecoderPlatformResourceAndroid::ESurfaceChangeResult::NewSurface)
-		{
-			// This depends on whether or not the decoder supports setting a new surface. If it doesn't we need to replay on a new instance.
-			check(!"Not currently handled");
-			bNeedReplayOnNewSurface = true;
-		}
-		else if (SurfaceChangeResult == IElectraDecoderResourceDelegateAndroid::IDecoderPlatformResourceAndroid::ESurfaceChangeResult::Error)
-		{
-			PostError(0, TEXT("Failed to get the new output surface"), ERRCODE_INTERNAL_INVALID_OUTPUT_SURFACE);
-			return IElectraDecoder::EDecoderError::Error;
-		}
-	}
-
 	// Resuming after suspending will re-create a decoder by replaying previous data.
 	// This may not be necessary on all devices, but is a safe catch-all.
-	if (DecoderCreatedAtResumeCount != CurrentResumeCount || bNeedReplayOnNewSurface)
+	if (DecoderCreatedAtResumeCount != CurrentResumeCount)
 	{
 		bTriggerReplay = true;
 		CurrentOutput.Reset();
 		InDecoderInput.Empty();
-		PendingDecoderOutputBuffers.Empty();
 		LastPushedPresentationTime = -1;
 		bNewDecoderRequired = true;
 	}
@@ -670,7 +599,6 @@ IElectraDecoder::EDecoderError FElectraVideoDecoderH265_Android::Flush()
 	}
 	CurrentOutput.Reset();
 	InDecoderInput.Empty();
-	PendingDecoderOutputBuffers.Empty();
 	LastPushedPresentationTime = -1;
 	bDidSendEOS = false;
 	DecodeState = EDecodeState::Decoding;
@@ -864,7 +792,6 @@ bool FElectraVideoDecoderH265_Android::InternalDecoderCreate()
 	bNewDecoderRequired = false;
 
 	CurrentSurfaceRequest.Reset();
-	CurrentBufferReleaseCallback.Reset();
 
 	DecodeState = EDecodeState::CreatingDecoder;
 	return true;
@@ -896,25 +823,13 @@ bool FElectraVideoDecoderH265_Android::InternalHandleDecoderCreate()
 			{
 				return InternalFinishDecoderCreate();
 			}
-			case FSurfaceRequestCallback::ESurfaceType::SurfaceView:
-			{
-			// TODO: When we are to use an externally provided surface we need to check with every decode call if it has changed.
-			//       If so we need to update the decoder, or if this fails re-create and replay it.
-				check(!"implement this!");
-				return PostError(0, TEXT("Failed to create decoder. Surface view is not currently supported."), ERRCODE_INTERNAL_COULD_NOT_CREATE_DECODER);
-			}
 		}
 	}
 	else
 	{
-		// We need a surface, for which we need to have a platform resource instance to get it.
-		if (!PlatformResource)
-		{
-			return PostError(0, TEXT("Failed to create decoder. No platform resource handle is set."), ERRCODE_INTERNAL_COULD_NOT_CREATE_DECODER);
-		}
 		// Create a request to get a surface and request one.
 		CurrentSurfaceRequest = MakeShared<FSurfaceRequestCallback, ESPMode::ThreadSafe>();
-		PlatformResource->RequestSurface(CurrentSurfaceRequest);
+		ResourceInstance->RequestSurface(CurrentSurfaceRequest);
 		return true;
 	}
 	return true;
@@ -984,12 +899,6 @@ bool FElectraVideoDecoderH265_Android::InternalFinishDecoderCreate()
 		}
 		DecodeState = EDecodeState::Decoding;
 
-		// Create a new buffer release interface.
-		if (PlatformResource)
-		{
-			CurrentBufferReleaseCallback = MakeShared<FBufferReleaseCallback, ESPMode::ThreadSafe>(StaticCastWeakPtr<FElectraVideoDecoderH265_Android>(AsWeak()));
-			PlatformResource->SetBufferReleaseCallback(CurrentBufferReleaseCallback);
-		}
 		return true;
 	}
 	else
@@ -1010,7 +919,6 @@ void FElectraVideoDecoderH265_Android::InternalDecoderDestroy()
 	DETAILLOG(LogElectraDecoders, Log, TEXT("VideoDecoderH265::InternalDecoderDestroy()"));
 	if (DecoderInstance.IsValid())
 	{
-		CurrentBufferReleaseCallback.Reset();
 		CurrentSurfaceRequest.Reset();
 		DecoderInstance->Flush();
 		DecoderInstance->Stop();
@@ -1020,36 +928,6 @@ void FElectraVideoDecoderH265_Android::InternalDecoderDestroy()
 }
 
 
-void FElectraVideoDecoderH265_Android::ReleaseOutputBuffer(const IElectraH265VideoDecoderAndroidJava::FOutputBufferInfo* InBufferInfo, bool bInRender, int64 InReleaseAt)
-{
-	if (InBufferInfo && DecoderInstance.IsValid() && HavePendingOutputIndex(InBufferInfo->BufferIndex, true))
-	{
-		int32 Result = DecoderInstance->ReleaseOutputBuffer(InBufferInfo->BufferIndex, InBufferInfo->ValidCount, bInRender, InReleaseAt);
-		if (Result)
-		{
-			PostError(Result, TEXT("Failed to release decoder output buffer"), ERRCODE_INTERNAL_FAILED_TO_CONVERT_OUTPUT_SAMPLE);
-		}
-	}
-}
-
-bool FElectraVideoDecoderH265_Android::HavePendingOutputIndex(int32 InIndex, bool bRemoveIfExists)
-{
-	for(int32 i=0; i<PendingDecoderOutputBuffers.Num(); ++i)
-	{
-		if (PendingDecoderOutputBuffers[i].BufferIndex == InIndex)
-		{
-			if (bRemoveIfExists)
-			{
-				PendingDecoderOutputBuffers.RemoveAtSwap(i);
-			}
-			return true;
-		}
-	}
-	return false;
-}
-
-
-
 FElectraVideoDecoderH265_Android::EConvertResult FElectraVideoDecoderH265_Android::ConvertDecoderOutput(const IElectraH265VideoDecoderAndroidJava::FOutputBufferInfo& InInfo)
 {
 	DETAILLOG(LogElectraDecoders, Log, TEXT("VideoDecoderH265::ConvertDecoderOutput()"));
@@ -1117,48 +995,20 @@ FElectraVideoDecoderH265_Android::EConvertResult FElectraVideoDecoderH265_Androi
 		NewOutput->FrameRateD = fr.Denom;
 	}
 
-	NewOutput->OwningDecoder = StaticCastWeakPtr<FElectraVideoDecoderH265_Android>(AsWeak());
-	NewOutput->OwningDecoderBufferInfo = InInfo;
-
 	NewOutput->ExtraValues.Emplace(TEXT("platform"), FVariant(TEXT("android")));
 	NewOutput->ExtraValues.Emplace(TEXT("codec"), FVariant(TEXT("hevc")));
 
-	PendingDecoderOutputBuffers.Emplace(InInfo);
 	CurrentOutput = MoveTemp(NewOutput);
-	return EConvertResult::Success;
-}
 
-
-void FElectraVideoDecoderOutputH265_Android::ReleaseOutputBuffer()
-{
-	if (OwningDecoderBufferInfo.BufferIndex >= 0)
+	if (DecoderInstance.IsValid())
 	{
-		TSharedPtr<FElectraVideoDecoderH265_Android, ESPMode::ThreadSafe> Decoder = OwningDecoder.Pin();
-		if (Decoder.IsValid())
+		int32 Result = DecoderInstance->ReleaseOutputBuffer(InInfo.BufferIndex, InInfo.ValidCount, true, -1);
+		if (Result)
 		{
-			Decoder->ReleaseOutputBuffer(&OwningDecoderBufferInfo, bBufferGotReferenced, -1);
+			PostError(Result, TEXT("Failed to release decoder output buffer"), ERRCODE_INTERNAL_FAILED_TO_CONVERT_OUTPUT_SAMPLE);
+			return EConvertResult::Failure;
 		}
-		OwningDecoderBufferInfo.BufferIndex = -1;
 	}
-}
-
-IElectraDecoderVideoOutput::EImageCopyResult FElectraVideoDecoderOutputH265_Android::CopyPlatformImage(IElectraDecoderVideoOutputCopyResources* InCopyResources) const
-{
-	TSharedPtr<FElectraVideoDecoderH265_Android, ESPMode::ThreadSafe> Decoder = OwningDecoder.Pin();
-	if (Decoder.IsValid())
-	{
-		bBufferGotReferenced = true;
-
-		InCopyResources->SetBufferIndex(OwningDecoderBufferInfo.BufferIndex);
-		InCopyResources->SetValidCount(OwningDecoderBufferInfo.ValidCount);
 
-		if (InCopyResources->ShouldReleaseBufferImmediately())
-		{
-			Decoder->ReleaseOutputBuffer(&OwningDecoderBufferInfo, bBufferGotReferenced, -1);
-			OwningDecoderBufferInfo.BufferIndex = -1;
-		}
-
-		return IElectraDecoderVideoOutput::EImageCopyResult::Succeeded;
-	}
-	return IElectraDecoderVideoOutput::EImageCopyResult::Failed;
+	return EConvertResult::Success;
 }
diff --git a/Engine/Plugins/Media/ElectraCodecs/Source/ElectraDecoders/Private/Android/h265/VideoDecoderH265_Android.h b/Engine/Plugins/Media/ElectraCodecs/Source/ElectraDecoders/Private/Android/h265/VideoDecoderH265_Android.h
--- a/Engine/Plugins/Media/ElectraCodecs/Source/ElectraDecoders/Private/Android/h265/VideoDecoderH265_Android.h
+++ b/Engine/Plugins/Media/ElectraCodecs/Source/ElectraDecoders/Private/Android/h265/VideoDecoderH265_Android.h
@@ -2,9 +2,7 @@
 
 #pragma once
 
-
 #include "IElectraDecoder.h"
-#include "IElectraDecoderResourceDelegate.h"
 
 class IElectraVideoDecoderH265_Android : public IElectraDecoder
 {
@@ -27,7 +25,7 @@ public:
 	static void PlatformGetSupportedConfigurations(TArray<FSupportedConfiguration>& OutSupportedConfigurations);
 
 	static void GetConfigurationOptions(TMap<FString, FVariant>& OutOptions);
-	static TSharedPtr<IElectraDecoder, ESPMode::ThreadSafe> Create(const TMap<FString, FVariant>& InOptions, TSharedPtr<IElectraDecoderResourceDelegate, ESPMode::ThreadSafe> InResourceDelegate);
+	static TSharedPtr<IElectraDecoder, ESPMode::ThreadSafe> Create(const TMap<FString, FVariant>& InOptions);
 
 	virtual ~IElectraVideoDecoderH265_Android() = default;
 };
diff --git a/Engine/Plugins/Media/ElectraCodecs/Source/ElectraDecoders/Private/Android/vpx/VPx_VideoDecoder_Android.cpp b/Engine/Plugins/Media/ElectraCodecs/Source/ElectraDecoders/Private/Android/vpx/VPx_VideoDecoder_Android.cpp
--- a/Engine/Plugins/Media/ElectraCodecs/Source/ElectraDecoders/Private/Android/vpx/VPx_VideoDecoder_Android.cpp
+++ b/Engine/Plugins/Media/ElectraCodecs/Source/ElectraDecoders/Private/Android/vpx/VPx_VideoDecoder_Android.cpp
@@ -41,9 +41,9 @@ public:
 		IElectraVideoDecoderVPx_Android::GetConfigurationOptions(OutOptions);
 	}
 
-	TSharedPtr<IElectraDecoder, ESPMode::ThreadSafe> CreateDecoderForFormat(const FString& InCodecFormat, const TMap<FString, FVariant>& InOptions, TSharedPtr<IElectraDecoderResourceDelegate, ESPMode::ThreadSafe> InResourceDelegate) override
+	TSharedPtr<IElectraDecoder, ESPMode::ThreadSafe> CreateDecoderForFormat(const FString& InCodecFormat, const TMap<FString, FVariant>& InOptions) override
 	{
-		return IElectraVideoDecoderVPx_Android::Create(InOptions, InResourceDelegate);
+		return IElectraVideoDecoderVPx_Android::Create(InOptions);
 	}
 
 };
