Electra: Removal of the platform resource delegate and the wrapping plugin

diff --git a/Engine/Plugins/Media/ElectraCodecs/Source/ElectraDecoders/Private/Android/vpx/VideoDecoderVPx_Android.cpp b/Engine/Plugins/Media/ElectraCodecs/Source/ElectraDecoders/Private/Android/vpx/VideoDecoderVPx_Android.cpp
--- a/Engine/Plugins/Media/ElectraCodecs/Source/ElectraDecoders/Private/Android/vpx/VideoDecoderVPx_Android.cpp
+++ b/Engine/Plugins/Media/ElectraCodecs/Source/ElectraDecoders/Private/Android/vpx/VideoDecoderVPx_Android.cpp
@@ -9,7 +9,7 @@
 #include "IElectraDecoderFeaturesAndOptions.h"
 #include "IElectraDecoderOutputVideo.h"
 
-#include "IElectraDecoderResourceDelegate.h"
+#include "ElectraDecodersPlatformResources.h"
 
 #include "ElectraDecodersModule.h"
 
@@ -33,21 +33,11 @@
 class FElectraVideoDecoderVPx_Android;
 
 
-class FElectraDecoderDefaultVideoOutputFormatVPx_Android : public IElectraDecoderDefaultVideoOutputFormat
-{
-public:
-	virtual ~FElectraDecoderDefaultVideoOutputFormatVPx_Android()
-	{ }
-};
-
-
 class FElectraVideoDecoderOutputVPx_Android : public IElectraDecoderVideoOutput
 {
 public:
 	virtual ~FElectraVideoDecoderOutputVPx_Android()
-	{
-		ReleaseOutputBuffer();
-	}
+	{ }
 
 	FTimespan GetPTS() const override
 	{ return PTS; }
@@ -80,13 +70,10 @@ public:
 	{ OutExtraValues = ExtraValues; }
 	void* GetPlatformOutputHandle(EElectraDecoderPlatformOutputHandleType InTypeOfHandle) const override
 	{ return nullptr; }
-	IElectraDecoderVideoOutputTransferHandle* GetTransferHandle() const override
-	{ return nullptr; }
-	IElectraDecoderVideoOutput::EImageCopyResult CopyPlatformImage(IElectraDecoderVideoOutputCopyResources* InCopyResources) const override;
+	IElectraDecoderVideoOutput::EImageCopyResult CopyPlatformImage(IElectraDecoderVideoOutputCopyResources* InCopyResources) const override
+	{ return IElectraDecoderVideoOutput::EImageCopyResult::NotSupported; }
 
 public:
-	void ReleaseOutputBuffer();
-
 	FTimespan PTS;
 	uint64 UserValue = 0;
 
@@ -101,10 +88,6 @@ public:
 	int32 FrameRateD = 0;
 	int32 PixelFormat = 0;
 	TMap<FString, FVariant> ExtraValues;
-
-	TWeakPtr<FElectraVideoDecoderVPx_Android, ESPMode::ThreadSafe> OwningDecoder;
-	mutable IElectraVPxVideoDecoderAndroidJava::FOutputBufferInfo OwningDecoderBufferInfo;
-	mutable bool bBufferGotReferenced = false;
 };
 
 
@@ -112,7 +95,7 @@ public:
 class FElectraVideoDecoderVPx_Android : public IElectraVideoDecoderVPx_Android
 {
 public:
-	FElectraVideoDecoderVPx_Android(const TMap<FString, FVariant>& InOptions, TSharedPtr<IElectraDecoderResourceDelegate, ESPMode::ThreadSafe> InResourceDelegate);
+	FElectraVideoDecoderVPx_Android(const TMap<FString, FVariant>& InOptions);
 
 	virtual ~FElectraVideoDecoderVPx_Android();
 
@@ -127,8 +110,6 @@ public:
 	IElectraDecoder::ECSDCompatibility IsCompatibleWith(const TMap<FString, FVariant>& CSDAndAdditionalOptions) override;
 	bool ResetToCleanStart() override;
 
-	TSharedPtr<IElectraDecoderDefaultOutputFormat, ESPMode::ThreadSafe> GetDefaultOutputFormatFromCSD(const TMap<FString, FVariant>& CSDAndAdditionalOptions) override;
-
 	EDecoderError DecodeAccessUnit(const FInputAccessUnit& InInputAccessUnit, const TMap<FString, FVariant>& InAdditionalOptions) override;
 	EDecoderError SendEndOfData() override;
 	EDecoderError Flush() override;
@@ -145,25 +126,8 @@ public:
 	void Suspend() override;
 	void Resume() override;
 
-	void ReleaseOutputBuffer(const IElectraVPxVideoDecoderAndroidJava::FOutputBufferInfo* InBufferInfo, bool bInRender, int64 InReleaseAt);
-
 private:
-	class FBufferReleaseCallback : public IElectraDecoderResourceDelegateAndroid::IDecoderPlatformResourceAndroid::IBufferReleaseCallback
-	{
-	public:
-		FBufferReleaseCallback(TWeakPtr<FElectraVideoDecoderVPx_Android, ESPMode::ThreadSafe> InOwningDecoder)
-			: OwningDecoder(MoveTemp(InOwningDecoder))
-		{ }
-		virtual ~FBufferReleaseCallback() = default;
-		void OnReleaseBuffer(int32 InBufferIndex, TOptional<int32> InBufferValidCount, TOptional<bool> InDoRender, TOptional<int64> InRenderAt) override
-		{
-			// Note: This could potentially be called from any thread but should be executed only from the decoder thread...
-			check(!"not implemented yet");
-		}
-		TWeakPtr<FElectraVideoDecoderVPx_Android, ESPMode::ThreadSafe> OwningDecoder;
-	};
-
-	class FSurfaceRequestCallback : public IElectraDecoderResourceDelegateAndroid::IDecoderPlatformResourceAndroid::ISurfaceRequestCallback
+	class FSurfaceRequestCallback : public FElectraDecodersPlatformResources::IDecoderResourceInstance::ISurfaceRequestCallback
 	{
 	public:
 		virtual ~FSurfaceRequestCallback() = default;
@@ -174,10 +138,10 @@ private:
 		jobject GetSurface() const
 		{ return Surface; }
 	private:
-		void OnNewSurface(ESurfaceType InSurfaceType, void* InSurface)
+		void OnNewSurface(ESurfaceType InSurfaceType, jobject InSurface)
 		{
 			SurfaceType = InSurfaceType;
-			Surface = (jobject) InSurface;
+			Surface = InSurface;
 			FPlatformMisc::MemoryBarrier();
 			bIsSet = true;
 		}
@@ -234,22 +198,19 @@ private:
 		GotEOS
 	};
 	EConvertResult ConvertDecoderOutput(const IElectraVPxVideoDecoderAndroidJava::FOutputBufferInfo& InInfo);
-	bool HavePendingOutputIndex(int32 InIndex, bool bRemoveIfExists);
 
 	IElectraDecoder::FError LastError;
 	FInitialMaxValues InitialMaxValues;
 	ECodec ActiveCodec = ECodec::VP9;
-
 	TMap<FString, FVariant> InitialCreationOptions;
-	TWeakPtr<IElectraDecoderResourceDelegate, ESPMode::ThreadSafe> ResourceDelegate;
-	IElectraDecoderResourceDelegateAndroid::IDecoderPlatformResourceAndroid* PlatformResource = nullptr;
+
+	FElectraDecodersPlatformResources::IDecoderResourceInstance* ResourceInstance = nullptr;
 
 	TSharedPtr<FElectraVideoDecoderOutputVPx_Android, ESPMode::ThreadSafe> CurrentOutput;
 	TArray<TSharedPtr<FDecoderInput, ESPMode::NotThreadSafe>> InDecoderInput;
 
 	TSharedPtr<IElectraVPxVideoDecoderAndroidJava, ESPMode::ThreadSafe> DecoderInstance;
 	IElectraVPxVideoDecoderAndroidJava::FDecoderInformation DecoderInfo;
-	TArray<IElectraVPxVideoDecoderAndroidJava::FOutputBufferInfo> PendingDecoderOutputBuffers;
 	EDecodeState DecodeState = EDecodeState::Decoding;
 	int64 LastPushedPresentationTime = -1;
 	bool bDidSendEOS = false;
@@ -258,7 +219,6 @@ private:
 	std::atomic<uint32> CurrentResumeCount;
 	bool bTriggerReplay = false;
 
-	TSharedPtr<FBufferReleaseCallback, ESPMode::ThreadSafe> CurrentBufferReleaseCallback;
 	TSharedPtr<FSurfaceRequestCallback, ESPMode::ThreadSafe> CurrentSurfaceRequest;
 	uint32 NativeDecoderID = 0;
 
@@ -303,22 +263,17 @@ void IElectraVideoDecoderVPx_Android::GetConfigurationOptions(TMap<FString, FVar
 	OutOptions.Emplace(IElectraDecoderFeature::SupportsDroppingOutput, FVariant(true));
 }
 
-TSharedPtr<IElectraDecoder, ESPMode::ThreadSafe> IElectraVideoDecoderVPx_Android::Create(const TMap<FString, FVariant>& InOptions, TSharedPtr<IElectraDecoderResourceDelegate, ESPMode::ThreadSafe> InResourceDelegate)
+TSharedPtr<IElectraDecoder, ESPMode::ThreadSafe> IElectraVideoDecoderVPx_Android::Create(const TMap<FString, FVariant>& InOptions)
 {
-	if (InResourceDelegate)
-	{
-		TSharedPtr<FElectraVideoDecoderVPx_Android, ESPMode::ThreadSafe> New = MakeShared<FElectraVideoDecoderVPx_Android>(InOptions, InResourceDelegate);
-		return New;
-	}
-	return nullptr;
+	TSharedPtr<FElectraVideoDecoderVPx_Android, ESPMode::ThreadSafe> New = MakeShared<FElectraVideoDecoderVPx_Android>(InOptions);
+	return New;
 }
 
 
-FElectraVideoDecoderVPx_Android::FElectraVideoDecoderVPx_Android(const TMap<FString, FVariant>& InOptions, TSharedPtr<IElectraDecoderResourceDelegate, ESPMode::ThreadSafe> InResourceDelegate)
+FElectraVideoDecoderVPx_Android::FElectraVideoDecoderVPx_Android(const TMap<FString, FVariant>& InOptions)
 	: InitialCreationOptions(InOptions)
-	, ResourceDelegate(InResourceDelegate)
 {
-	PlatformResource = reinterpret_cast<IElectraDecoderResourceDelegateAndroid::IDecoderPlatformResourceAndroid*>(ElectraDecodersUtil::GetVariantValueSafeU64(InOptions, TEXT("platformresource"), 0));
+	ResourceInstance = FElectraDecodersPlatformResources::CreateResourceInstance(ElectraDecodersUtil::GetVariantValueSafeI64(InOptions, FElectraDecodersPlatformResources::OptionKeyOutputTexturePoolID().ToString(), 0));
 
 	uint32 Codec4CC = (uint32) ElectraDecodersUtil::GetVariantValueSafeI64(InOptions, TEXT("codec_4cc"), 0);
 	ActiveCodec = Codec4CC == 'vp09' ? ECodec::VP9 : ECodec::VP8;
@@ -360,6 +315,8 @@ FElectraVideoDecoderVPx_Android::~FElectraVideoDecoderVPx_Android()
 	check(LastError.Code == ERRCODE_INTERNAL_ALREADY_CLOSED);
 	// We do it nonetheless...
 	Close();
+
+	FElectraDecodersPlatformResources::ReleaseResourceInstance(ResourceInstance);
 }
 
 
@@ -401,7 +358,6 @@ bool FElectraVideoDecoderVPx_Android::ResetToCleanStart()
 
 	CurrentOutput.Reset();
 	InDecoderInput.Empty();
-	PendingDecoderOutputBuffers.Empty();
 	DecodeState = EDecodeState::Decoding;
 	LastPushedPresentationTime = -1;
 	bDidSendEOS = false;
@@ -410,12 +366,6 @@ bool FElectraVideoDecoderVPx_Android::ResetToCleanStart()
 }
 
 
-TSharedPtr<IElectraDecoderDefaultOutputFormat, ESPMode::ThreadSafe> FElectraVideoDecoderVPx_Android::GetDefaultOutputFormatFromCSD(const TMap<FString, FVariant>& CSDAndAdditionalOptions)
-{
-	return nullptr;
-}
-
-
 IElectraDecoder::EDecoderError FElectraVideoDecoderVPx_Android::DecodeAccessUnit(const FInputAccessUnit& InInputAccessUnit, const TMap<FString, FVariant>& InAdditionalOptions)
 {
 	// If already in error do nothing!
@@ -452,34 +402,13 @@ IElectraDecoder::EDecoderError FElectraVideoDecoderVPx_Android::DecodeAccessUnit
 		}
 	}
 
-	// When using a SurfaceView we need to check if it did not suddenly change.
-	bool bNeedReplayOnNewSurface = false;
-	if (PlatformResource && CurrentSurfaceRequest.IsValid() && CurrentSurfaceRequest->HasCompleted() && CurrentSurfaceRequest->GetType() == FSurfaceRequestCallback::ESurfaceType::SurfaceView)
-	{
-		void* NewSurfaceView = nullptr;
-		IElectraDecoderResourceDelegateAndroid::IDecoderPlatformResourceAndroid::ESurfaceChangeResult SurfaceChangeResult;
-		SurfaceChangeResult = PlatformResource->VerifySurfaceView(NewSurfaceView, reinterpret_cast<void*>(CurrentSurfaceRequest->GetSurface()));
-		if (SurfaceChangeResult == IElectraDecoderResourceDelegateAndroid::IDecoderPlatformResourceAndroid::ESurfaceChangeResult::NewSurface)
-		{
-			// This depends on whether or not the decoder supports setting a new surface. If it doesn't we need to replay on a new instance.
-			check(!"Not currently handled");
-			bNeedReplayOnNewSurface = true;
-		}
-		else if (SurfaceChangeResult == IElectraDecoderResourceDelegateAndroid::IDecoderPlatformResourceAndroid::ESurfaceChangeResult::Error)
-		{
-			PostError(0, TEXT("Failed to get the new output surface"), ERRCODE_INTERNAL_INVALID_OUTPUT_SURFACE);
-			return IElectraDecoder::EDecoderError::Error;
-		}
-	}
-
 	// Resuming after suspending will re-create a decoder by replaying previous data.
 	// This may not be necessary on all devices, but is a safe catch-all.
-	if (DecoderCreatedAtResumeCount != CurrentResumeCount || bNeedReplayOnNewSurface)
+	if (DecoderCreatedAtResumeCount != CurrentResumeCount)
 	{
 		bTriggerReplay = true;
 		CurrentOutput.Reset();
 		InDecoderInput.Empty();
-		PendingDecoderOutputBuffers.Empty();
 		LastPushedPresentationTime = -1;
 		bNewDecoderRequired = true;
 	}
@@ -632,7 +561,6 @@ IElectraDecoder::EDecoderError FElectraVideoDecoderVPx_Android::Flush()
 	}
 	CurrentOutput.Reset();
 	InDecoderInput.Empty();
-	PendingDecoderOutputBuffers.Empty();
 	LastPushedPresentationTime = -1;
 	bDidSendEOS = false;
 	DecodeState = EDecodeState::Decoding;
@@ -826,7 +754,6 @@ bool FElectraVideoDecoderVPx_Android::InternalDecoderCreate()
 	bNewDecoderRequired = false;
 
 	CurrentSurfaceRequest.Reset();
-	CurrentBufferReleaseCallback.Reset();
 
 	DecodeState = EDecodeState::CreatingDecoder;
 	return true;
@@ -858,25 +785,13 @@ bool FElectraVideoDecoderVPx_Android::InternalHandleDecoderCreate()
 			{
 				return InternalFinishDecoderCreate();
 			}
-			case FSurfaceRequestCallback::ESurfaceType::SurfaceView:
-			{
-			// TODO: When we are to use an externally provided surface we need to check with every decode call if it has changed.
-			//       If so we need to update the decoder, or if this fails re-create and replay it.
-				check(!"implement this!");
-				return PostError(0, TEXT("Failed to create decoder. Surface view is not currently supported."), ERRCODE_INTERNAL_COULD_NOT_CREATE_DECODER);
-			}
 		}
 	}
 	else
 	{
-		// We need a surface, for which we need to have a platform resource instance to get it.
-		if (!PlatformResource)
-		{
-			return PostError(0, TEXT("Failed to create decoder. No platform resource handle is set."), ERRCODE_INTERNAL_COULD_NOT_CREATE_DECODER);
-		}
 		// Create a request to get a surface and request one.
 		CurrentSurfaceRequest = MakeShared<FSurfaceRequestCallback, ESPMode::ThreadSafe>();
-		PlatformResource->RequestSurface(CurrentSurfaceRequest);
+		ResourceInstance->RequestSurface(CurrentSurfaceRequest);
 		return true;
 	}
 	return true;
@@ -924,12 +839,6 @@ bool FElectraVideoDecoderVPx_Android::InternalFinishDecoderCreate()
 		}
 		DecodeState = EDecodeState::Decoding;
 
-		// Create a new buffer release interface.
-		if (PlatformResource)
-		{
-			CurrentBufferReleaseCallback = MakeShared<FBufferReleaseCallback, ESPMode::ThreadSafe>(StaticCastWeakPtr<FElectraVideoDecoderVPx_Android>(AsWeak()));
-			PlatformResource->SetBufferReleaseCallback(CurrentBufferReleaseCallback);
-		}
 		return true;
 	}
 	else
@@ -950,7 +859,6 @@ void FElectraVideoDecoderVPx_Android::InternalDecoderDestroy()
 	DETAILLOG(LogElectraDecoders, Log, TEXT("VideoDecoderVPx::InternalDecoderDestroy()"));
 	if (DecoderInstance.IsValid())
 	{
-		CurrentBufferReleaseCallback.Reset();
 		CurrentSurfaceRequest.Reset();
 		DecoderInstance->Flush();
 		DecoderInstance->Stop();
@@ -960,36 +868,6 @@ void FElectraVideoDecoderVPx_Android::InternalDecoderDestroy()
 }
 
 
-void FElectraVideoDecoderVPx_Android::ReleaseOutputBuffer(const IElectraVPxVideoDecoderAndroidJava::FOutputBufferInfo* InBufferInfo, bool bInRender, int64 InReleaseAt)
-{
-	if (InBufferInfo && DecoderInstance.IsValid() && HavePendingOutputIndex(InBufferInfo->BufferIndex, true))
-	{
-		int32 Result = DecoderInstance->ReleaseOutputBuffer(InBufferInfo->BufferIndex, InBufferInfo->ValidCount, bInRender, InReleaseAt);
-		if (Result)
-		{
-			PostError(Result, TEXT("Failed to release decoder output buffer"), ERRCODE_INTERNAL_FAILED_TO_CONVERT_OUTPUT_SAMPLE);
-		}
-	}
-}
-
-bool FElectraVideoDecoderVPx_Android::HavePendingOutputIndex(int32 InIndex, bool bRemoveIfExists)
-{
-	for(int32 i=0; i<PendingDecoderOutputBuffers.Num(); ++i)
-	{
-		if (PendingDecoderOutputBuffers[i].BufferIndex == InIndex)
-		{
-			if (bRemoveIfExists)
-			{
-				PendingDecoderOutputBuffers.RemoveAtSwap(i);
-			}
-			return true;
-		}
-	}
-	return false;
-}
-
-
-
 FElectraVideoDecoderVPx_Android::EConvertResult FElectraVideoDecoderVPx_Android::ConvertDecoderOutput(const IElectraVPxVideoDecoderAndroidJava::FOutputBufferInfo& InInfo)
 {
 	DETAILLOG(LogElectraDecoders, Log, TEXT("VideoDecoderVPx::ConvertDecoderOutput()"));
@@ -1050,48 +928,21 @@ FElectraVideoDecoderVPx_Android::EConvertResult FElectraVideoDecoderVPx_Android:
 	NewOutput->Crop.Bottom = Info.Height - Info.CropBottom - 1;
 	NewOutput->AspectW = 1;
 	NewOutput->AspectH = 1;
-	NewOutput->OwningDecoder = StaticCastWeakPtr<FElectraVideoDecoderVPx_Android>(AsWeak());
-	NewOutput->OwningDecoderBufferInfo = InInfo;
 
 	NewOutput->ExtraValues.Emplace(TEXT("platform"), FVariant(TEXT("android")));
 	NewOutput->ExtraValues.Emplace(TEXT("codec"), FVariant(ActiveCodec == ECodec::VP9 ? TEXT("vp9") : TEXT("vp8")));
 
-	PendingDecoderOutputBuffers.Emplace(InInfo);
 	CurrentOutput = MoveTemp(NewOutput);
-	return EConvertResult::Success;
-}
 
-
-void FElectraVideoDecoderOutputVPx_Android::ReleaseOutputBuffer()
-{
-	if (OwningDecoderBufferInfo.BufferIndex >= 0)
+	if (DecoderInstance.IsValid())
 	{
-		TSharedPtr<FElectraVideoDecoderVPx_Android, ESPMode::ThreadSafe> Decoder = OwningDecoder.Pin();
-		if (Decoder.IsValid())
+		int32 Result = DecoderInstance->ReleaseOutputBuffer(InInfo.BufferIndex, InInfo.ValidCount, true, -1);
+		if (Result)
 		{
-			Decoder->ReleaseOutputBuffer(&OwningDecoderBufferInfo, bBufferGotReferenced, -1);
+			PostError(Result, TEXT("Failed to release decoder output buffer"), ERRCODE_INTERNAL_FAILED_TO_CONVERT_OUTPUT_SAMPLE);
+			return EConvertResult::Failure;
 		}
-		OwningDecoderBufferInfo.BufferIndex = -1;
 	}
-}
-
-IElectraDecoderVideoOutput::EImageCopyResult FElectraVideoDecoderOutputVPx_Android::CopyPlatformImage(IElectraDecoderVideoOutputCopyResources* InCopyResources) const
-{
-	TSharedPtr<FElectraVideoDecoderVPx_Android, ESPMode::ThreadSafe> Decoder = OwningDecoder.Pin();
-	if (Decoder.IsValid())
-	{
-		bBufferGotReferenced = true;
-
-		InCopyResources->SetBufferIndex(OwningDecoderBufferInfo.BufferIndex);
-		InCopyResources->SetValidCount(OwningDecoderBufferInfo.ValidCount);
 
-		if (InCopyResources->ShouldReleaseBufferImmediately())
-		{
-			Decoder->ReleaseOutputBuffer(&OwningDecoderBufferInfo, bBufferGotReferenced, -1);
-			OwningDecoderBufferInfo.BufferIndex = -1;
-		}
-
-		return IElectraDecoderVideoOutput::EImageCopyResult::Succeeded;
-	}
-	return IElectraDecoderVideoOutput::EImageCopyResult::Failed;
+	return EConvertResult::Success;
 }
diff --git a/Engine/Plugins/Media/ElectraCodecs/Source/ElectraDecoders/Private/Apple/aac/AAC_AudioDecoder_Apple.cpp b/Engine/Plugins/Media/ElectraCodecs/Source/ElectraDecoders/Private/Apple/aac/AAC_AudioDecoder_Apple.cpp
--- a/Engine/Plugins/Media/ElectraCodecs/Source/ElectraDecoders/Private/Apple/aac/AAC_AudioDecoder_Apple.cpp
+++ b/Engine/Plugins/Media/ElectraCodecs/Source/ElectraDecoders/Private/Apple/aac/AAC_AudioDecoder_Apple.cpp
@@ -71,9 +71,9 @@ public:
 		IElectraAudioDecoderAAC_Apple::GetConfigurationOptions(OutOptions);
 	}
 
-	TSharedPtr<IElectraDecoder, ESPMode::ThreadSafe> CreateDecoderForFormat(const FString& InCodecFormat, const TMap<FString, FVariant>& InOptions, TSharedPtr<IElectraDecoderResourceDelegate, ESPMode::ThreadSafe> InResourceDelegate) override
+	TSharedPtr<IElectraDecoder, ESPMode::ThreadSafe> CreateDecoderForFormat(const FString& InCodecFormat, const TMap<FString, FVariant>& InOptions) override
 	{
-		return IElectraAudioDecoderAAC_Apple::Create(InOptions, InResourceDelegate);
+		return IElectraAudioDecoderAAC_Apple::Create(InOptions);
 	}
 
 };
diff --git a/Engine/Plugins/Media/ElectraCodecs/Source/ElectraDecoders/Private/Apple/aac/AudioDecoderAAC_Apple.cpp b/Engine/Plugins/Media/ElectraCodecs/Source/ElectraDecoders/Private/Apple/aac/AudioDecoderAAC_Apple.cpp
--- a/Engine/Plugins/Media/ElectraCodecs/Source/ElectraDecoders/Private/Apple/aac/AudioDecoderAAC_Apple.cpp
+++ b/Engine/Plugins/Media/ElectraCodecs/Source/ElectraDecoders/Private/Apple/aac/AudioDecoderAAC_Apple.cpp
@@ -12,33 +12,12 @@
 #include "IElectraDecoderFeaturesAndOptions.h"
 #include "IElectraDecoderOutputAudio.h"
 
-#include "IElectraDecoderResourceDelegate.h"
-
 /*********************************************************************************************************************/
 #include <AudioToolbox/AudioToolbox.h>
 
 #define MAGIC_INSUFFICIENT_INPUT_MARKER		4711
 /*********************************************************************************************************************/
 
-class FElectraDecoderDefaultAudioOutputFormatAAC_Apple : public IElectraDecoderDefaultAudioOutputFormat
-{
-public:
-	virtual ~FElectraDecoderDefaultAudioOutputFormatAAC_Apple()
-	{ }
-
-	int32 GetNumChannels() const override
-	{ return NumChannels; }
-	int32 GetSampleRate() const override
-	{ return SampleRate; }
-	int32 GetNumFrames() const override
-	{ return NumFrames; }
-
-	int32 NumChannels = 0;
-	int32 SampleRate = 0;
-	int32 NumFrames = 0;
-};
-
-
 class FElectraAudioDecoderOutputAAC_Apple : public IElectraDecoderAudioOutput
 {
 public:
@@ -90,7 +69,7 @@ public:
 class FElectraAudioDecoderAAC_Apple : public IElectraAudioDecoderAAC_Apple
 {
 public:
-	FElectraAudioDecoderAAC_Apple(const TMap<FString, FVariant>& InOptions, TSharedPtr<IElectraDecoderResourceDelegate, ESPMode::ThreadSafe> InResourceDelegate);
+	FElectraAudioDecoderAAC_Apple(const TMap<FString, FVariant>& InOptions);
 
 	virtual ~FElectraAudioDecoderAAC_Apple();
 
@@ -105,8 +84,6 @@ public:
 	IElectraDecoder::ECSDCompatibility IsCompatibleWith(const TMap<FString, FVariant>& CSDAndAdditionalOptions) override;
 	bool ResetToCleanStart() override;
 
-	TSharedPtr<IElectraDecoderDefaultOutputFormat, ESPMode::ThreadSafe> GetDefaultOutputFormatFromCSD(const TMap<FString, FVariant>& CSDAndAdditionalOptions) override;
-
 	EDecoderError DecodeAccessUnit(const FInputAccessUnit& InInputAccessUnit, const TMap<FString, FVariant>& InAdditionalOptions) override;
 	EDecoderError SendEndOfData() override;
 	EDecoderError Flush() override;
@@ -185,8 +162,6 @@ private:
 
 	bool ConvertDecoderOutput(int32 InNumSamplesProduced);
 
-	TWeakPtr<IElectraDecoderResourceDelegate, ESPMode::ThreadSafe> ResourceDelegate;
-
 	IElectraDecoder::FError LastError;
 
 	TSharedPtr<ElectraDecodersUtil::MPEG::FAACDecoderConfigurationRecord, ESPMode::ThreadSafe> ConfigRecord;
@@ -288,21 +263,15 @@ void IElectraAudioDecoderAAC_Apple::GetConfigurationOptions(TMap<FString, FVaria
 {
 }
 
-TSharedPtr<IElectraDecoder, ESPMode::ThreadSafe> IElectraAudioDecoderAAC_Apple::Create(const TMap<FString, FVariant>& InOptions, TSharedPtr<IElectraDecoderResourceDelegate, ESPMode::ThreadSafe> InResourceDelegate)
+TSharedPtr<IElectraDecoder, ESPMode::ThreadSafe> IElectraAudioDecoderAAC_Apple::Create(const TMap<FString, FVariant>& InOptions)
 {
-	if (InResourceDelegate)
-	{
-		TSharedPtr<FElectraAudioDecoderAAC_Apple, ESPMode::ThreadSafe> New = MakeShared<FElectraAudioDecoderAAC_Apple>(InOptions, InResourceDelegate);
-		return New;
-	}
-	return nullptr;
+	TSharedPtr<FElectraAudioDecoderAAC_Apple, ESPMode::ThreadSafe> New = MakeShared<FElectraAudioDecoderAAC_Apple>(InOptions);
+	return New;
 }
 
 
-FElectraAudioDecoderAAC_Apple::FElectraAudioDecoderAAC_Apple(const TMap<FString, FVariant>& InOptions, TSharedPtr<IElectraDecoderResourceDelegate, ESPMode::ThreadSafe> InResourceDelegate)
+FElectraAudioDecoderAAC_Apple::FElectraAudioDecoderAAC_Apple(const TMap<FString, FVariant>& InOptions)
 {
-	ResourceDelegate = InResourceDelegate;
-
 	// If there is codec specific data set we create a config record from it now.
 	ConfigRecord = GetCSDFromOptions(InOptions);
 
@@ -390,21 +359,6 @@ bool FElectraAudioDecoderAAC_Apple::ResetToCleanStart()
 }
 
 
-TSharedPtr<IElectraDecoderDefaultOutputFormat, ESPMode::ThreadSafe> FElectraAudioDecoderAAC_Apple::GetDefaultOutputFormatFromCSD(const TMap<FString, FVariant>& CSDAndAdditionalOptions)
-{
-	TSharedPtr<ElectraDecodersUtil::MPEG::FAACDecoderConfigurationRecord, ESPMode::ThreadSafe> Cfg = GetCSDFromOptions(CSDAndAdditionalOptions);
-	if (Cfg.IsValid())
-	{
-		TSharedPtr<FElectraDecoderDefaultAudioOutputFormatAAC_Apple, ESPMode::ThreadSafe> Info = MakeShared<FElectraDecoderDefaultAudioOutputFormatAAC_Apple, ESPMode::ThreadSafe>();
-		Info->NumChannels = Cfg->PSSignal > 0 ? 2 : NumChannelsForConfig[Cfg->ChannelConfiguration];
-		Info->SampleRate = Cfg->ExtSamplingFrequency ? Cfg->ExtSamplingFrequency : ConfigRecord->SamplingRate;
-		Info->NumFrames = Cfg->SBRSignal > 0 ? 2048 : 1024;
-		return Info;
-	}
-	return nullptr;
-}
-
-
 IElectraDecoder::EDecoderError FElectraAudioDecoderAAC_Apple::DecodeAccessUnit(const FInputAccessUnit& InInputAccessUnit, const TMap<FString, FVariant>& InAdditionalOptions)
 {
 	// If already in error do nothing!
diff --git a/Engine/Plugins/Media/ElectraCodecs/Source/ElectraDecoders/Private/Apple/aac/AudioDecoderAAC_Apple.h b/Engine/Plugins/Media/ElectraCodecs/Source/ElectraDecoders/Private/Apple/aac/AudioDecoderAAC_Apple.h
--- a/Engine/Plugins/Media/ElectraCodecs/Source/ElectraDecoders/Private/Apple/aac/AudioDecoderAAC_Apple.h
+++ b/Engine/Plugins/Media/ElectraCodecs/Source/ElectraDecoders/Private/Apple/aac/AudioDecoderAAC_Apple.h
@@ -3,7 +3,6 @@
 #pragma once
 
 #include "IElectraDecoder.h"
-#include "IElectraDecoderResourceDelegate.h"
 
 #ifdef ELECTRA_DECODERS_ENABLE_APPLE
 
@@ -11,7 +10,7 @@ class IElectraAudioDecoderAAC_Apple : public IElectraDecoder
 {
 public:
 	static void GetConfigurationOptions(TMap<FString, FVariant>& OutOptions);
-	static TSharedPtr<IElectraDecoder, ESPMode::ThreadSafe> Create(const TMap<FString, FVariant>& InOptions, TSharedPtr<IElectraDecoderResourceDelegate, ESPMode::ThreadSafe> InResourceDelegate);
+	static TSharedPtr<IElectraDecoder, ESPMode::ThreadSafe> Create(const TMap<FString, FVariant>& InOptions);
 
 	virtual ~IElectraAudioDecoderAAC_Apple() = default;
 };
diff --git a/Engine/Plugins/Media/ElectraCodecs/Source/ElectraDecoders/Private/Apple/h264/H264_VideoDecoder_Apple.cpp b/Engine/Plugins/Media/ElectraCodecs/Source/ElectraDecoders/Private/Apple/h264/H264_VideoDecoder_Apple.cpp
--- a/Engine/Plugins/Media/ElectraCodecs/Source/ElectraDecoders/Private/Apple/h264/H264_VideoDecoder_Apple.cpp
+++ b/Engine/Plugins/Media/ElectraCodecs/Source/ElectraDecoders/Private/Apple/h264/H264_VideoDecoder_Apple.cpp
@@ -83,9 +83,9 @@ public:
 		IElectraVideoDecoderH264_Apple::GetConfigurationOptions(OutOptions);
 	}
 
-	TSharedPtr<IElectraDecoder, ESPMode::ThreadSafe> CreateDecoderForFormat(const FString& InCodecFormat, const TMap<FString, FVariant>& InOptions, TSharedPtr<IElectraDecoderResourceDelegate, ESPMode::ThreadSafe> InResourceDelegate) override
+	TSharedPtr<IElectraDecoder, ESPMode::ThreadSafe> CreateDecoderForFormat(const FString& InCodecFormat, const TMap<FString, FVariant>& InOptions) override
 	{
-		return IElectraVideoDecoderH264_Apple::Create(InOptions, InResourceDelegate);
+		return IElectraVideoDecoderH264_Apple::Create(InOptions);
 	}
 
 };
diff --git a/Engine/Plugins/Media/ElectraCodecs/Source/ElectraDecoders/Private/Apple/h264/VideoDecoderH264_Apple.cpp b/Engine/Plugins/Media/ElectraCodecs/Source/ElectraDecoders/Private/Apple/h264/VideoDecoderH264_Apple.cpp
--- a/Engine/Plugins/Media/ElectraCodecs/Source/ElectraDecoders/Private/Apple/h264/VideoDecoderH264_Apple.cpp
+++ b/Engine/Plugins/Media/ElectraCodecs/Source/ElectraDecoders/Private/Apple/h264/VideoDecoderH264_Apple.cpp
@@ -12,7 +12,6 @@
 #include "IElectraDecoderFeaturesAndOptions.h"
 #include "IElectraDecoderOutputVideo.h"
 
-#include "IElectraDecoderResourceDelegate.h"
 #include "ElectraDecodersModule.h"
 
 #include "Apple/AppleElectraDecoderPlatformOutputFormatTypes.h"
@@ -24,14 +23,6 @@
 class FElectraVideoDecoderH264_Apple;
 
 
-class FElectraDecoderDefaultVideoOutputFormatH264_Apple : public IElectraDecoderDefaultVideoOutputFormat
-{
-public:
-	virtual ~FElectraDecoderDefaultVideoOutputFormatH264_Apple()
-	{ }
-};
-
-
 class FElectraVideoDecoderOutputH264_Apple : public IElectraDecoderVideoOutput, public IElectraDecoderVideoOutputImageBuffers
 {
 public:
@@ -77,8 +68,6 @@ public:
 		}
 		return nullptr;
 	}
-	IElectraDecoderVideoOutputTransferHandle* GetTransferHandle() const override
-	{ return nullptr; }
 	IElectraDecoderVideoOutput::EImageCopyResult CopyPlatformImage(IElectraDecoderVideoOutputCopyResources* InCopyResources) const override
 	{ return IElectraDecoderVideoOutput::EImageCopyResult::NotSupported; }
 
@@ -154,7 +143,7 @@ class FElectraVideoDecoderH264_Apple : public IElectraVideoDecoderH264_Apple
 	static constexpr int32 kNumReorderFrames = 5;
 
 public:
-	FElectraVideoDecoderH264_Apple(const TMap<FString, FVariant>& InOptions, TSharedPtr<IElectraDecoderResourceDelegate, ESPMode::ThreadSafe> InResourceDelegate);
+	FElectraVideoDecoderH264_Apple(const TMap<FString, FVariant>& InOptions);
 
 	virtual ~FElectraVideoDecoderH264_Apple();
 
@@ -169,8 +158,6 @@ public:
 	IElectraDecoder::ECSDCompatibility IsCompatibleWith(const TMap<FString, FVariant>& CSDAndAdditionalOptions) override;
 	bool ResetToCleanStart() override;
 
-	TSharedPtr<IElectraDecoderDefaultOutputFormat, ESPMode::ThreadSafe> GetDefaultOutputFormatFromCSD(const TMap<FString, FVariant>& CSDAndAdditionalOptions) override;
-
 	EDecoderError DecodeAccessUnit(const FInputAccessUnit& InInputAccessUnit, const TMap<FString, FVariant>& InAdditionalOptions) override;
 	EDecoderError SendEndOfData() override;
 	EDecoderError Flush() override;
@@ -335,7 +322,6 @@ private:
 	EConvertResult ConvertDecoderOutput();
 
 	TMap<FString, FVariant> InitialCreationOptions;
-	TWeakPtr<IElectraDecoderResourceDelegate, ESPMode::ThreadSafe> ResourceDelegate;
 
 	IElectraDecoder::FError LastError;
 
@@ -396,20 +382,15 @@ void IElectraVideoDecoderH264_Apple::GetConfigurationOptions(TMap<FString, FVari
 	OutOptions.Emplace(IElectraDecoderFeature::SupportsDroppingOutput, FVariant(true));
 }
 
-TSharedPtr<IElectraDecoder, ESPMode::ThreadSafe> IElectraVideoDecoderH264_Apple::Create(const TMap<FString, FVariant>& InOptions, TSharedPtr<IElectraDecoderResourceDelegate, ESPMode::ThreadSafe> InResourceDelegate)
+TSharedPtr<IElectraDecoder, ESPMode::ThreadSafe> IElectraVideoDecoderH264_Apple::Create(const TMap<FString, FVariant>& InOptions)
 {
-	if (InResourceDelegate.IsValid())
-	{
-		TSharedPtr<FElectraVideoDecoderH264_Apple, ESPMode::ThreadSafe> New = MakeShared<FElectraVideoDecoderH264_Apple>(InOptions, InResourceDelegate);
-		return New;
-	}
-	return nullptr;
+	TSharedPtr<FElectraVideoDecoderH264_Apple, ESPMode::ThreadSafe> New = MakeShared<FElectraVideoDecoderH264_Apple>(InOptions);
+	return New;
 }
 
-FElectraVideoDecoderH264_Apple::FElectraVideoDecoderH264_Apple(const TMap<FString, FVariant>& InOptions, TSharedPtr<IElectraDecoderResourceDelegate, ESPMode::ThreadSafe> InResourceDelegate)
+FElectraVideoDecoderH264_Apple::FElectraVideoDecoderH264_Apple(const TMap<FString, FVariant>& InOptions)
 {
 	InitialCreationOptions = InOptions;
-	ResourceDelegate = InResourceDelegate;
 }
 
 FElectraVideoDecoderH264_Apple::~FElectraVideoDecoderH264_Apple()
@@ -522,11 +503,6 @@ void FElectraVideoDecoderH264_Apple::InternalFlushAllInputAndOutput()
 }
 
 
-TSharedPtr<IElectraDecoderDefaultOutputFormat, ESPMode::ThreadSafe> FElectraVideoDecoderH264_Apple::GetDefaultOutputFormatFromCSD(const TMap<FString, FVariant>& CSDAndAdditionalOptions)
-{
-	return nullptr;
-}
-
 IElectraDecoder::ECSDCompatibility FElectraVideoDecoderH264_Apple::IsCompatibleWith(const TMap<FString, FVariant>& CSDAndAdditionalOptions)
 {
 	// When we have no decoder yet then we are compatible because we will be creating a decoder when needed.
diff --git a/Engine/Plugins/Media/ElectraCodecs/Source/ElectraDecoders/Private/Apple/h264/VideoDecoderH264_Apple.h b/Engine/Plugins/Media/ElectraCodecs/Source/ElectraDecoders/Private/Apple/h264/VideoDecoderH264_Apple.h
--- a/Engine/Plugins/Media/ElectraCodecs/Source/ElectraDecoders/Private/Apple/h264/VideoDecoderH264_Apple.h
+++ b/Engine/Plugins/Media/ElectraCodecs/Source/ElectraDecoders/Private/Apple/h264/VideoDecoderH264_Apple.h
@@ -3,7 +3,6 @@
 #pragma once
 
 #include "IElectraDecoder.h"
-#include "IElectraDecoderResourceDelegate.h"
 
 #ifdef ELECTRA_DECODERS_ENABLE_APPLE
 
@@ -26,7 +25,7 @@ public:
 	static void PlatformGetSupportedConfigurations(TArray<FSupportedConfiguration>& OutSupportedConfigurations);
 
 	static void GetConfigurationOptions(TMap<FString, FVariant>& OutOptions);
-	static TSharedPtr<IElectraDecoder, ESPMode::ThreadSafe> Create(const TMap<FString, FVariant>& InOptions, TSharedPtr<IElectraDecoderResourceDelegate, ESPMode::ThreadSafe> InResourceDelegate);
+	static TSharedPtr<IElectraDecoder, ESPMode::ThreadSafe> Create(const TMap<FString, FVariant>& InOptions);
 
 	virtual ~IElectraVideoDecoderH264_Apple() = default;
 };
diff --git a/Engine/Plugins/Media/ElectraCodecs/Source/ElectraDecoders/Private/Apple/h265/H265_VideoDecoder_Apple.cpp b/Engine/Plugins/Media/ElectraCodecs/Source/ElectraDecoders/Private/Apple/h265/H265_VideoDecoder_Apple.cpp
--- a/Engine/Plugins/Media/ElectraCodecs/Source/ElectraDecoders/Private/Apple/h265/H265_VideoDecoder_Apple.cpp
+++ b/Engine/Plugins/Media/ElectraCodecs/Source/ElectraDecoders/Private/Apple/h265/H265_VideoDecoder_Apple.cpp
@@ -83,9 +83,9 @@ public:
 		IElectraVideoDecoderH265_Apple::GetConfigurationOptions(OutOptions);
 	}
 
-	TSharedPtr<IElectraDecoder, ESPMode::ThreadSafe> CreateDecoderForFormat(const FString& InCodecFormat, const TMap<FString, FVariant>& InOptions, TSharedPtr<IElectraDecoderResourceDelegate, ESPMode::ThreadSafe> InResourceDelegate) override
+	TSharedPtr<IElectraDecoder, ESPMode::ThreadSafe> CreateDecoderForFormat(const FString& InCodecFormat, const TMap<FString, FVariant>& InOptions) override
 	{
-		return IElectraVideoDecoderH265_Apple::Create(InOptions, InResourceDelegate);
+		return IElectraVideoDecoderH265_Apple::Create(InOptions);
 	}
 
 };
diff --git a/Engine/Plugins/Media/ElectraCodecs/Source/ElectraDecoders/Private/Apple/h265/VideoDecoderH265_Apple.cpp b/Engine/Plugins/Media/ElectraCodecs/Source/ElectraDecoders/Private/Apple/h265/VideoDecoderH265_Apple.cpp
--- a/Engine/Plugins/Media/ElectraCodecs/Source/ElectraDecoders/Private/Apple/h265/VideoDecoderH265_Apple.cpp
+++ b/Engine/Plugins/Media/ElectraCodecs/Source/ElectraDecoders/Private/Apple/h265/VideoDecoderH265_Apple.cpp
@@ -12,7 +12,6 @@
 #include "IElectraDecoderFeaturesAndOptions.h"
 #include "IElectraDecoderOutputVideo.h"
 
-#include "IElectraDecoderResourceDelegate.h"
 #include "ElectraDecodersModule.h"
 
 #include "Apple/AppleElectraDecoderPlatformOutputFormatTypes.h"
@@ -24,14 +23,6 @@
 class FElectraVideoDecoderH265_Apple;
 
 
-class FElectraDecoderDefaultVideoOutputFormatH265_Apple : public IElectraDecoderDefaultVideoOutputFormat
-{
-public:
-	virtual ~FElectraDecoderDefaultVideoOutputFormatH265_Apple()
-	{ }
-};
-
-
 class FElectraVideoDecoderOutputH265_Apple : public IElectraDecoderVideoOutput, public IElectraDecoderVideoOutputImageBuffers
 {
 public:
@@ -77,8 +68,6 @@ public:
 		}
 		return nullptr;
 	}
-	IElectraDecoderVideoOutputTransferHandle* GetTransferHandle() const override
-	{ return nullptr; }
 	IElectraDecoderVideoOutput::EImageCopyResult CopyPlatformImage(IElectraDecoderVideoOutputCopyResources* InCopyResources) const override
 	{ return IElectraDecoderVideoOutput::EImageCopyResult::NotSupported; }
 
@@ -154,7 +143,7 @@ class FElectraVideoDecoderH265_Apple : public IElectraVideoDecoderH265_Apple
 	static constexpr int32 kNumReorderFrames = 5;
 
 public:
-	FElectraVideoDecoderH265_Apple(const TMap<FString, FVariant>& InOptions, TSharedPtr<IElectraDecoderResourceDelegate, ESPMode::ThreadSafe> InResourceDelegate);
+	FElectraVideoDecoderH265_Apple(const TMap<FString, FVariant>& InOptions);
 
 	virtual ~FElectraVideoDecoderH265_Apple();
 
@@ -169,8 +158,6 @@ public:
 	IElectraDecoder::ECSDCompatibility IsCompatibleWith(const TMap<FString, FVariant>& CSDAndAdditionalOptions) override;
 	bool ResetToCleanStart() override;
 
-	TSharedPtr<IElectraDecoderDefaultOutputFormat, ESPMode::ThreadSafe> GetDefaultOutputFormatFromCSD(const TMap<FString, FVariant>& CSDAndAdditionalOptions) override;
-
 	EDecoderError DecodeAccessUnit(const FInputAccessUnit& InInputAccessUnit, const TMap<FString, FVariant>& InAdditionalOptions) override;
 	EDecoderError SendEndOfData() override;
 	EDecoderError Flush() override;
@@ -335,7 +322,6 @@ private:
 	EConvertResult ConvertDecoderOutput();
 
 	TMap<FString, FVariant> InitialCreationOptions;
-	TWeakPtr<IElectraDecoderResourceDelegate, ESPMode::ThreadSafe> ResourceDelegate;
 
 	IElectraDecoder::FError LastError;
 
@@ -396,20 +382,15 @@ void IElectraVideoDecoderH265_Apple::GetConfigurationOptions(TMap<FString, FVari
 	OutOptions.Emplace(IElectraDecoderFeature::SupportsDroppingOutput, FVariant(true));
 }
 
-TSharedPtr<IElectraDecoder, ESPMode::ThreadSafe> IElectraVideoDecoderH265_Apple::Create(const TMap<FString, FVariant>& InOptions, TSharedPtr<IElectraDecoderResourceDelegate, ESPMode::ThreadSafe> InResourceDelegate)
+TSharedPtr<IElectraDecoder, ESPMode::ThreadSafe> IElectraVideoDecoderH265_Apple::Create(const TMap<FString, FVariant>& InOptions)
 {
-	if (InResourceDelegate.IsValid())
-	{
-		TSharedPtr<FElectraVideoDecoderH265_Apple, ESPMode::ThreadSafe> New = MakeShared<FElectraVideoDecoderH265_Apple>(InOptions, InResourceDelegate);
-		return New;
-	}
-	return nullptr;
+	TSharedPtr<FElectraVideoDecoderH265_Apple, ESPMode::ThreadSafe> New = MakeShared<FElectraVideoDecoderH265_Apple>(InOptions);
+	return New;
 }
 
-FElectraVideoDecoderH265_Apple::FElectraVideoDecoderH265_Apple(const TMap<FString, FVariant>& InOptions, TSharedPtr<IElectraDecoderResourceDelegate, ESPMode::ThreadSafe> InResourceDelegate)
+FElectraVideoDecoderH265_Apple::FElectraVideoDecoderH265_Apple(const TMap<FString, FVariant>& InOptions)
 {
 	InitialCreationOptions = InOptions;
-	ResourceDelegate = InResourceDelegate;
 }
 
 FElectraVideoDecoderH265_Apple::~FElectraVideoDecoderH265_Apple()
@@ -525,11 +506,6 @@ void FElectraVideoDecoderH265_Apple::InternalFlushAllInputAndOutput()
 }
 
 
-TSharedPtr<IElectraDecoderDefaultOutputFormat, ESPMode::ThreadSafe> FElectraVideoDecoderH265_Apple::GetDefaultOutputFormatFromCSD(const TMap<FString, FVariant>& CSDAndAdditionalOptions)
-{
-	return nullptr;
-}
-
 IElectraDecoder::ECSDCompatibility FElectraVideoDecoderH265_Apple::IsCompatibleWith(const TMap<FString, FVariant>& CSDAndAdditionalOptions)
 {
 	// When we have no decoder yet then we are compatible because we will be creating a decoder when needed.
diff --git a/Engine/Plugins/Media/ElectraCodecs/Source/ElectraDecoders/Private/DX/AudioDecoderAAC_DX.cpp b/Engine/Plugins/Media/ElectraCodecs/Source/ElectraDecoders/Private/DX/AudioDecoderAAC_DX.cpp
--- a/Engine/Plugins/Media/ElectraCodecs/Source/ElectraDecoders/Private/DX/AudioDecoderAAC_DX.cpp
+++ b/Engine/Plugins/Media/ElectraCodecs/Source/ElectraDecoders/Private/DX/AudioDecoderAAC_DX.cpp
@@ -24,25 +24,6 @@ if (FAILED(Result))						\
 
 /*********************************************************************************************************************/
 
-class FElectraDecoderDefaultAudioOutputFormatAAC_DX : public IElectraDecoderDefaultAudioOutputFormat
-{
-public:
-	virtual ~FElectraDecoderDefaultAudioOutputFormatAAC_DX()
-	{ }
-
-	int32 GetNumChannels() const override
-	{ return NumChannels; }
-	int32 GetSampleRate() const override
-	{ return SampleRate; }
-	int32 GetNumFrames() const override
-	{ return NumFrames; }
-
-	int32 NumChannels = 0;
-	int32 SampleRate = 0;
-	int32 NumFrames = 0;
-};
-
-
 class FElectraAudioDecoderOutputAAC_DX : public IElectraDecoderAudioOutput
 {
 public:
@@ -94,7 +75,7 @@ public:
 class FElectraAudioDecoderAAC_DX : public IElectraAudioDecoderAAC_DX
 {
 public:
-	FElectraAudioDecoderAAC_DX(const TMap<FString, FVariant>& InOptions, TSharedPtr<IElectraDecoderResourceDelegate, ESPMode::ThreadSafe> InResourceDelegate);
+	FElectraAudioDecoderAAC_DX(const TMap<FString, FVariant>& InOptions);
 
 	virtual ~FElectraAudioDecoderAAC_DX();
 
@@ -109,8 +90,6 @@ public:
 	IElectraDecoder::ECSDCompatibility IsCompatibleWith(const TMap<FString, FVariant>& CSDAndAdditionalOptions) override;
 	bool ResetToCleanStart() override;
 
-	TSharedPtr<IElectraDecoderDefaultOutputFormat, ESPMode::ThreadSafe> GetDefaultOutputFormatFromCSD(const TMap<FString, FVariant>& CSDAndAdditionalOptions) override;
-
 	EDecoderError DecodeAccessUnit(const FInputAccessUnit& InInputAccessUnit, const TMap<FString, FVariant>& InAdditionalOptions) override;
 	EDecoderError SendEndOfData() override;
 	EDecoderError Flush() override;
@@ -201,8 +180,6 @@ private:
 	bool ConvertDecoderOutput();
 
 
-	TWeakPtr<IElectraDecoderResourceDelegate, ESPMode::ThreadSafe> ResourceDelegate;
-
 	IElectraDecoder::FError LastError;
 
 	TSharedPtr<ElectraDecodersUtil::MPEG::FAACDecoderConfigurationRecord, ESPMode::ThreadSafe> ConfigRecord;
@@ -278,17 +255,15 @@ void IElectraAudioDecoderAAC_DX::GetConfigurationOptions(TMap<FString, FVariant>
 }
 
 
-TSharedPtr<IElectraDecoder, ESPMode::ThreadSafe> IElectraAudioDecoderAAC_DX::Create(const TMap<FString, FVariant>& InOptions, TSharedPtr<IElectraDecoderResourceDelegate, ESPMode::ThreadSafe> InResourceDelegate)
+TSharedPtr<IElectraDecoder, ESPMode::ThreadSafe> IElectraAudioDecoderAAC_DX::Create(const TMap<FString, FVariant>& InOptions)
 {
-	TSharedPtr<FElectraAudioDecoderAAC_DX, ESPMode::ThreadSafe> New = MakeShared<FElectraAudioDecoderAAC_DX>(InOptions, InResourceDelegate);
+	TSharedPtr<FElectraAudioDecoderAAC_DX, ESPMode::ThreadSafe> New = MakeShared<FElectraAudioDecoderAAC_DX>(InOptions);
 	return New;
 }
 
 
-FElectraAudioDecoderAAC_DX::FElectraAudioDecoderAAC_DX(const TMap<FString, FVariant>& InOptions, TSharedPtr<IElectraDecoderResourceDelegate, ESPMode::ThreadSafe> InResourceDelegate)
+FElectraAudioDecoderAAC_DX::FElectraAudioDecoderAAC_DX(const TMap<FString, FVariant>& InOptions)
 {
-	ResourceDelegate = InResourceDelegate;
-
 	// If there is codec specific data set we create a config record from it now.
 	ConfigRecord = GetCSDFromOptions(InOptions);
 }
@@ -365,22 +340,6 @@ bool FElectraAudioDecoderAAC_DX::ResetToCleanStart()
 }
 
 
-TSharedPtr<IElectraDecoderDefaultOutputFormat, ESPMode::ThreadSafe> FElectraAudioDecoderAAC_DX::GetDefaultOutputFormatFromCSD(const TMap<FString, FVariant>& CSDAndAdditionalOptions)
-{
-	TSharedPtr<ElectraDecodersUtil::MPEG::FAACDecoderConfigurationRecord, ESPMode::ThreadSafe> Cfg = GetCSDFromOptions(CSDAndAdditionalOptions);
-	if (Cfg.IsValid())
-	{
-		const uint8 NumChannelsForConfig[16] = { 0, 1, 2, 3, 4, 5, 6, 7, 0, 0, 0, 7, 8, 0, 8, 0 };
-		TSharedPtr<FElectraDecoderDefaultAudioOutputFormatAAC_DX, ESPMode::ThreadSafe> Info = MakeShared<FElectraDecoderDefaultAudioOutputFormatAAC_DX, ESPMode::ThreadSafe>();
-		Info->NumChannels = Cfg->PSSignal > 0 ? 2 : NumChannelsForConfig[Cfg->ChannelConfiguration];
-		Info->SampleRate = Cfg->ExtSamplingFrequency ? Cfg->ExtSamplingFrequency : ConfigRecord->SamplingRate;
-		Info->NumFrames = Cfg->SBRSignal > 0 ? 2048 : 1024;
-		return Info;
-	}
-	return nullptr;
-}
-
-
 IElectraDecoder::ECSDCompatibility FElectraAudioDecoderAAC_DX::IsCompatibleWith(const TMap<FString, FVariant>& CSDAndAdditionalOptions)
 {
 	return IElectraDecoder::ECSDCompatibility::DrainAndReset;
diff --git a/Engine/Plugins/Media/ElectraCodecs/Source/ElectraDecoders/Private/DX/VideoDecoderH264_DX.cpp b/Engine/Plugins/Media/ElectraCodecs/Source/ElectraDecoders/Private/DX/VideoDecoderH264_DX.cpp
--- a/Engine/Plugins/Media/ElectraCodecs/Source/ElectraDecoders/Private/DX/VideoDecoderH264_DX.cpp
+++ b/Engine/Plugins/Media/ElectraCodecs/Source/ElectraDecoders/Private/DX/VideoDecoderH264_DX.cpp
@@ -12,9 +12,9 @@
 #include "IElectraDecoderFeaturesAndOptions.h"
 #include "IElectraDecoderOutputVideo.h"
 
-#include "IElectraDecoderResourceDelegate.h"
 #include "ElectraDecodersModule.h"
 
+#include "ElectraDecodersPlatformResources.h"
 #include COMPILED_PLATFORM_HEADER(ElectraDecoderGPUBufferHelpers.h)
 #include COMPILED_PLATFORM_HEADER(PlatformHeaders_Video_DX.h)
 
@@ -55,15 +55,6 @@ static TAutoConsoleVariable<int32> CVarElectraWindowsH264UseOldOutputPath(
 /*********************************************************************************************************************/
 /*********************************************************************************************************************/
 
-class FElectraDecoderDefaultVideoOutputFormatH264_DX : public IElectraDecoderDefaultVideoOutputFormat
-{
-public:
-	virtual ~FElectraDecoderDefaultVideoOutputFormatH264_DX()
-	{ }
-
-};
-
-
 class FElectraVideoDecoderOutputH264_DX : public IElectraDecoderVideoOutput
 #if ELECTRA_HAVE_IMAGEBUFFERS
 										, public IElectraDecoderVideoOutputImageBuffers
@@ -127,8 +118,6 @@ public:
 				return nullptr;
 		}
 	}
-	IElectraDecoderVideoOutputTransferHandle* GetTransferHandle() const override
-	{ return nullptr; }
 	IElectraDecoderVideoOutput::EImageCopyResult CopyPlatformImage(IElectraDecoderVideoOutputCopyResources* InCopyResources) const override
 	{ return IElectraDecoderVideoOutput::EImageCopyResult::NotSupported; }
 
@@ -240,7 +229,7 @@ public:
 	} Buffer;
 #else
 	FElectraMediaDecoderOutputBufferPool_DX12::FOutputData Buffer;
-	TSharedPtr<IElectraDecoderResourceDelegateWindows::IAsyncConsecutiveTaskSync, ESPMode::ThreadSafe> Buffer_TaskSync;
+	TSharedPtr<FElectraDecodersPlatformResources::FAsyncConsecutiveTaskSync, ESPMode::ThreadSafe> Buffer_TaskSync;
 #endif
 };
 
@@ -249,7 +238,7 @@ public:
 class FElectraVideoDecoderH264_DX : public IElectraVideoDecoderH264_DX
 {
 public:
-	FElectraVideoDecoderH264_DX(const TMap<FString, FVariant>& InOptions, TSharedPtr<IElectraDecoderResourceDelegate, ESPMode::ThreadSafe> InResourceDelegate);
+	FElectraVideoDecoderH264_DX(const TMap<FString, FVariant>& InOptions);
 
 	virtual ~FElectraVideoDecoderH264_DX();
 
@@ -264,8 +253,6 @@ public:
 	IElectraDecoder::ECSDCompatibility IsCompatibleWith(const TMap<FString, FVariant>& CSDAndAdditionalOptions) override;
 	bool ResetToCleanStart() override;
 
-	TSharedPtr<IElectraDecoderDefaultOutputFormat, ESPMode::ThreadSafe> GetDefaultOutputFormatFromCSD(const TMap<FString, FVariant>& CSDAndAdditionalOptions) override;
-
 	EDecoderError DecodeAccessUnit(const FInputAccessUnit& InInputAccessUnit, const TMap<FString, FVariant>& InAdditionalOptions) override;
 	EDecoderError SendEndOfData() override;
 	EDecoderError Flush() override;
@@ -362,8 +349,6 @@ private:
 	bool ConvertDecoderOutput();
 	static void CopyBufferData(TRefCountPtr<IMF2DBuffer> Buffer2D, TRefCountPtr<ID3D12Resource> BufferResource, TRefCountPtr<ID3D12Fence> BufferFence, uint64 BufferFenceValue, uint32 BufferPitch, int32 DecodedHeight);
 
-	TWeakPtr<IElectraDecoderResourceDelegate, ESPMode::ThreadSafe> ResourceDelegate;
-
 	IElectraDecoder::FError LastError;
 
 	TUniquePtr<FDecoderOutputBuffer> CurrentDecoderOutputBuffer;
@@ -384,7 +369,7 @@ private:
 	uint32 MaxHeight;
 #if !ALLOW_MFSAMPLE_WITH_DX12
 	mutable TSharedPtr<FElectraMediaDecoderOutputBufferPool_DX12> D3D12ResourcePool;
-	mutable TSharedPtr<IElectraDecoderResourceDelegateBase::IAsyncConsecutiveTaskSync> TaskSync;
+	mutable TSharedPtr<FElectraDecodersPlatformResources::FAsyncConsecutiveTaskSync, ESPMode::ThreadSafe> TaskSync;
 #endif
 
 	static const GUID MFTms_AVLowLatencyMode;
@@ -394,20 +379,15 @@ const GUID FElectraVideoDecoderH264_DX::MFTms_AVLowLatencyMode = { 0x9c27891a, 0
 
 
 
-TSharedPtr<IElectraDecoder, ESPMode::ThreadSafe> IElectraVideoDecoderH264_DX::Create(const TMap<FString, FVariant>& InOptions, TSharedPtr<IElectraDecoderResourceDelegate, ESPMode::ThreadSafe> InResourceDelegate)
+TSharedPtr<IElectraDecoder, ESPMode::ThreadSafe> IElectraVideoDecoderH264_DX::Create(const TMap<FString, FVariant>& InOptions)
 {
-	if (InResourceDelegate.IsValid())
-	{
-		TSharedPtr<FElectraVideoDecoderH264_DX, ESPMode::ThreadSafe> New = MakeShared<FElectraVideoDecoderH264_DX>(InOptions, InResourceDelegate);
-		return New;
-	}
-	return nullptr;
+	TSharedPtr<FElectraVideoDecoderH264_DX, ESPMode::ThreadSafe> New = MakeShared<FElectraVideoDecoderH264_DX>(InOptions);
+	return New;
 }
 
 
-FElectraVideoDecoderH264_DX::FElectraVideoDecoderH264_DX(const TMap<FString, FVariant>& InOptions, TSharedPtr<IElectraDecoderResourceDelegate, ESPMode::ThreadSafe> InResourceDelegate)
-	: ResourceDelegate(InResourceDelegate)
-	, InitialCreationOptions(InOptions)
+FElectraVideoDecoderH264_DX::FElectraVideoDecoderH264_DX(const TMap<FString, FVariant>& InOptions)
+	: InitialCreationOptions(InOptions)
 {
 	MaxWidth = Align((uint32)ElectraDecodersUtil::GetVariantValueSafeU64(InOptions, TEXT("max_width"), 1920), 16u);
 	MaxHeight = Align((uint32)ElectraDecodersUtil::GetVariantValueSafeU64(InOptions, TEXT("max_height"), 1080), 16u);
@@ -505,12 +485,6 @@ bool FElectraVideoDecoderH264_DX::ResetToCleanStart()
 }
 
 
-TSharedPtr<IElectraDecoderDefaultOutputFormat, ESPMode::ThreadSafe> FElectraVideoDecoderH264_DX::GetDefaultOutputFormatFromCSD(const TMap<FString, FVariant>& CSDAndAdditionalOptions)
-{
-	return nullptr;
-}
-
-
 IElectraDecoder::ECSDCompatibility FElectraVideoDecoderH264_DX::IsCompatibleWith(const TMap<FString, FVariant>& CSDAndAdditionalOptions)
 {
 	// When we have no decoder yet then we are compatible because we will be creating a decoder when needed.
@@ -817,7 +791,7 @@ bool FElectraVideoDecoderH264_DX::InternalDecoderCreate(const TMap<FString, FVar
 
 	// Create decoder transform
 	TMap<FString, FVariant> NoOptions;
-	LastError = PlatformCreateMFDecoderTransform(&DecoderPlatformHandle, ResourceDelegate.Pin(), NoOptions);
+	LastError = PlatformCreateMFDecoderTransform(&DecoderPlatformHandle, NoOptions);
 	if (LastError.IsSet())
 	{
 		return false;
@@ -1153,10 +1127,9 @@ bool FElectraVideoDecoderH264_DX::ConvertDecoderOutput()
 	void* PlatformDevice = nullptr;
 	int32 PlatformDeviceVersion = 0;
 	bool bUseGPUBuffers = false;
-	auto PinnedResourceDelegate = ResourceDelegate.Pin();
-	if (PinnedResourceDelegate.IsValid() && CVarElectraWindowsH264UseOldOutputPath.GetValueOnAnyThread() == 0)
+	if (CVarElectraWindowsH264UseOldOutputPath.GetValueOnAnyThread() == 0)
 	{
-		PinnedResourceDelegate->GetD3DDevice(&PlatformDevice, &PlatformDeviceVersion);
+		FElectraDecodersPlatformResources::GetD3DDeviceAndVersion(&PlatformDevice, &PlatformDeviceVersion);
 		bUseGPUBuffers = (PlatformDevice && PlatformDeviceVersion >= 12000);
 	}
 
@@ -1193,7 +1166,7 @@ bool FElectraVideoDecoderH264_DX::ConvertDecoderOutput()
 		// Create a tasksync instance so we can have our own async jobs run in consecutive order
 		if (!TaskSync.IsValid())
 		{
-			TaskSync = PinnedResourceDelegate->CreateAsyncConsecutiveTaskSync();
+			TaskSync = FElectraDecodersPlatformResources::CreateAsyncConsecutiveTaskSync();
 		}
 
 		// Get data from MFSample...
@@ -1212,10 +1185,10 @@ bool FElectraVideoDecoderH264_DX::ConvertDecoderOutput()
 
 		NewOutput->Pitch = BufferPitch;
 
-		if (PinnedResourceDelegate->RunCodeAsync([DecodedOutputSample, Buffer2D, BufferResource = NewOutput->Buffer.Resource, BufferFence = NewOutput->Buffer.Fence, BufferFenceValue = NewOutput->Buffer.FenceValue, BufferPitch, DecodedHeight = NewOutput->DecodedHeight]()
+		if (FElectraDecodersPlatformResources::RunCodeAsync([DecodedOutputSample, Buffer2D, BufferResource = NewOutput->Buffer.Resource, BufferFence = NewOutput->Buffer.Fence, BufferFenceValue = NewOutput->Buffer.FenceValue, BufferPitch, DecodedHeight = NewOutput->DecodedHeight]()
 			{
 				CopyBufferData(Buffer2D, BufferResource, BufferFence, BufferFenceValue, BufferPitch, DecodedHeight);
-			}, TaskSync.Get()))
+			}, TaskSync))
 		{
 			NewOutput->Buffer_TaskSync = TaskSync;
 		}
diff --git a/Engine/Plugins/Media/ElectraCodecs/Source/ElectraDecoders/Private/DX/VideoDecoderH264_DX.h b/Engine/Plugins/Media/ElectraCodecs/Source/ElectraDecoders/Private/DX/VideoDecoderH264_DX.h
--- a/Engine/Plugins/Media/ElectraCodecs/Source/ElectraDecoders/Private/DX/VideoDecoderH264_DX.h
+++ b/Engine/Plugins/Media/ElectraCodecs/Source/ElectraDecoders/Private/DX/VideoDecoderH264_DX.h
@@ -5,7 +5,6 @@
 #ifdef ELECTRA_DECODERS_ENABLE_DX
 
 #include "IElectraDecoder.h"
-#include "IElectraDecoderResourceDelegate.h"
 
 class IElectraVideoDecoderH264_DX : public IElectraDecoder
 {
@@ -40,11 +39,11 @@ public:
 		virtual void* GetDXDeviceContext() const = 0;
 		virtual bool IsSoftware() const = 0;
 	};
-	static FError PlatformCreateMFDecoderTransform(IPlatformHandle** OutPlatformHandle, TSharedPtr<IElectraDecoderResourceDelegate, ESPMode::ThreadSafe> InResourceDelegate, const TMap<FString, FVariant>& InOptions);
+	static FError PlatformCreateMFDecoderTransform(IPlatformHandle** OutPlatformHandle, const TMap<FString, FVariant>& InOptions);
 	static void PlatformReleaseMFDecoderTransform(IPlatformHandle** OutPlatformHandle);
 
 	
-	static TSharedPtr<IElectraDecoder, ESPMode::ThreadSafe> Create(const TMap<FString, FVariant>& InOptions, TSharedPtr<IElectraDecoderResourceDelegate, ESPMode::ThreadSafe> InResourceDelegate);
+	static TSharedPtr<IElectraDecoder, ESPMode::ThreadSafe> Create(const TMap<FString, FVariant>& InOptions);
 
 	virtual ~IElectraVideoDecoderH264_DX() = default;
 };
diff --git a/Engine/Plugins/Media/ElectraCodecs/Source/ElectraDecoders/Private/DX/VideoDecoderH265_DX.cpp b/Engine/Plugins/Media/ElectraCodecs/Source/ElectraDecoders/Private/DX/VideoDecoderH265_DX.cpp
--- a/Engine/Plugins/Media/ElectraCodecs/Source/ElectraDecoders/Private/DX/VideoDecoderH265_DX.cpp
+++ b/Engine/Plugins/Media/ElectraCodecs/Source/ElectraDecoders/Private/DX/VideoDecoderH265_DX.cpp
@@ -11,9 +11,9 @@
 #include "IElectraDecoderFeaturesAndOptions.h"
 #include "IElectraDecoderOutputVideo.h"
 
-#include "IElectraDecoderResourceDelegate.h"
 #include "ElectraDecodersModule.h"
 
+#include "ElectraDecodersPlatformResources.h"
 #include COMPILED_PLATFORM_HEADER(ElectraDecoderGPUBufferHelpers.h)
 #include COMPILED_PLATFORM_HEADER(PlatformHeaders_Video_DX.h)
 
@@ -54,15 +54,6 @@ static TAutoConsoleVariable<int32> CVarElectraWindowsH265UseOldOutputPath(
 /*********************************************************************************************************************/
 /*********************************************************************************************************************/
 
-class FElectraDecoderDefaultVideoOutputFormatH265_DX : public IElectraDecoderDefaultVideoOutputFormat
-{
-public:
-	virtual ~FElectraDecoderDefaultVideoOutputFormatH265_DX()
-	{ }
-
-};
-
-
 class FElectraVideoDecoderOutputH265_DX : public IElectraDecoderVideoOutput
 #if ELECTRA_HAVE_IMAGEBUFFERS
 	, public IElectraDecoderVideoOutputImageBuffers
@@ -126,8 +117,6 @@ public:
 			return nullptr;
 		}
 	}
-	IElectraDecoderVideoOutputTransferHandle* GetTransferHandle() const override
-	{ return nullptr; }
 	IElectraDecoderVideoOutput::EImageCopyResult CopyPlatformImage(IElectraDecoderVideoOutputCopyResources* InCopyResources) const override
 	{ return IElectraDecoderVideoOutput::EImageCopyResult::NotSupported; }
 
@@ -239,7 +228,7 @@ public:
 	} Buffer;
 #else
 	FElectraMediaDecoderOutputBufferPool_DX12::FOutputData Buffer;
-	TSharedPtr<IElectraDecoderResourceDelegateWindows::IAsyncConsecutiveTaskSync, ESPMode::ThreadSafe> Buffer_TaskSync;
+	TSharedPtr<FElectraDecodersPlatformResources::FAsyncConsecutiveTaskSync, ESPMode::ThreadSafe> Buffer_TaskSync;
 #endif
 };
 
@@ -248,7 +237,7 @@ public:
 class FElectraVideoDecoderH265_DX : public IElectraVideoDecoderH265_DX
 {
 public:
-	FElectraVideoDecoderH265_DX(const TMap<FString, FVariant>& InOptions, TSharedPtr<IElectraDecoderResourceDelegate, ESPMode::ThreadSafe> InResourceDelegate);
+	FElectraVideoDecoderH265_DX(const TMap<FString, FVariant>& InOptions);
 
 	virtual ~FElectraVideoDecoderH265_DX();
 
@@ -263,8 +252,6 @@ public:
 	IElectraDecoder::ECSDCompatibility IsCompatibleWith(const TMap<FString, FVariant>& CSDAndAdditionalOptions) override;
 	bool ResetToCleanStart() override;
 
-	TSharedPtr<IElectraDecoderDefaultOutputFormat, ESPMode::ThreadSafe> GetDefaultOutputFormatFromCSD(const TMap<FString, FVariant>& CSDAndAdditionalOptions) override;
-
 	EDecoderError DecodeAccessUnit(const FInputAccessUnit& InInputAccessUnit, const TMap<FString, FVariant>& InAdditionalOptions) override;
 	EDecoderError SendEndOfData() override;
 	EDecoderError Flush() override;
@@ -360,8 +347,6 @@ private:
 	bool ConvertDecoderOutput();
 	static void CopyBufferData(TRefCountPtr<IMF2DBuffer> Buffer2D, TRefCountPtr<ID3D12Resource> BufferResource, TRefCountPtr<ID3D12Fence> BufferFence, uint64 BufferFenceValue, uint32 BufferPitch, int32 DecodedHeight);
 
-	TWeakPtr<IElectraDecoderResourceDelegate, ESPMode::ThreadSafe> ResourceDelegate;
-
 	IElectraDecoder::FError LastError;
 
 	TUniquePtr<FDecoderOutputBuffer> CurrentDecoderOutputBuffer;
@@ -381,7 +366,7 @@ private:
 	uint32 MaxHeight;
 #if !ALLOW_MFSAMPLE_WITH_DX12
 	mutable TSharedPtr<FElectraMediaDecoderOutputBufferPool_DX12> D3D12ResourcePool;
-	mutable TSharedPtr<IElectraDecoderResourceDelegateBase::IAsyncConsecutiveTaskSync> TaskSync;
+	mutable TSharedPtr<FElectraDecodersPlatformResources::FAsyncConsecutiveTaskSync, ESPMode::ThreadSafe> TaskSync;
 #endif
 
 	static const GUID MFTms_AVLowLatencyMode;
@@ -391,20 +376,15 @@ const GUID FElectraVideoDecoderH265_DX::MFTms_AVLowLatencyMode = { 0x9c27891a, 0
 
 
 
-TSharedPtr<IElectraDecoder, ESPMode::ThreadSafe> IElectraVideoDecoderH265_DX::Create(const TMap<FString, FVariant>& InOptions, TSharedPtr<IElectraDecoderResourceDelegate, ESPMode::ThreadSafe> InResourceDelegate)
+TSharedPtr<IElectraDecoder, ESPMode::ThreadSafe> IElectraVideoDecoderH265_DX::Create(const TMap<FString, FVariant>& InOptions)
 {
-	if (InResourceDelegate.IsValid())
-	{
-		TSharedPtr<FElectraVideoDecoderH265_DX, ESPMode::ThreadSafe> New = MakeShared<FElectraVideoDecoderH265_DX>(InOptions, InResourceDelegate);
-		return New;
-	}
-	return nullptr;
+	TSharedPtr<FElectraVideoDecoderH265_DX, ESPMode::ThreadSafe> New = MakeShared<FElectraVideoDecoderH265_DX>(InOptions);
+	return New;
 }
 
 
-FElectraVideoDecoderH265_DX::FElectraVideoDecoderH265_DX(const TMap<FString, FVariant>& InOptions, TSharedPtr<IElectraDecoderResourceDelegate, ESPMode::ThreadSafe> InResourceDelegate)
-	: ResourceDelegate(InResourceDelegate)
-	, InitialCreationOptions(InOptions)
+FElectraVideoDecoderH265_DX::FElectraVideoDecoderH265_DX(const TMap<FString, FVariant>& InOptions)
+	: InitialCreationOptions(InOptions)
 {
 	MaxWidth = Align((uint32)ElectraDecodersUtil::GetVariantValueSafeU64(InOptions, TEXT("max_width"), 3840), 16u);
 	MaxHeight = Align((uint32)ElectraDecodersUtil::GetVariantValueSafeU64(InOptions, TEXT("max_height"), 2160), 16u);
@@ -488,12 +468,6 @@ bool FElectraVideoDecoderH265_DX::ResetToCleanStart()
 }
 
 
-TSharedPtr<IElectraDecoderDefaultOutputFormat, ESPMode::ThreadSafe> FElectraVideoDecoderH265_DX::GetDefaultOutputFormatFromCSD(const TMap<FString, FVariant>& CSDAndAdditionalOptions)
-{
-	return nullptr;
-}
-
-
 IElectraDecoder::ECSDCompatibility FElectraVideoDecoderH265_DX::IsCompatibleWith(const TMap<FString, FVariant>& CSDAndAdditionalOptions)
 {
 	// When we have no decoder yet then we are compatible because we will be creating a decoder when needed.
@@ -790,7 +764,7 @@ bool FElectraVideoDecoderH265_DX::InternalDecoderCreate(const TMap<FString, FVar
 
 	// Create decoder transform
 	TMap<FString, FVariant> NoOptions;
-	LastError = PlatformCreateMFDecoderTransform(&DecoderPlatformHandle, ResourceDelegate.Pin(), NoOptions);
+	LastError = PlatformCreateMFDecoderTransform(&DecoderPlatformHandle, NoOptions);
 	if (LastError.IsSet())
 	{
 		return false;
@@ -1159,10 +1133,9 @@ bool FElectraVideoDecoderH265_DX::ConvertDecoderOutput()
 	void* PlatformDevice = nullptr;
 	int32 PlatformDeviceVersion = 0;
 	bool bUseGPUBuffers = false;
-	auto PinnedResourceDelegate = ResourceDelegate.Pin();
-	if (PinnedResourceDelegate.IsValid() && CVarElectraWindowsH265UseOldOutputPath.GetValueOnAnyThread() == 0)
+	if (CVarElectraWindowsH265UseOldOutputPath.GetValueOnAnyThread() == 0)
 	{
-		PinnedResourceDelegate->GetD3DDevice(&PlatformDevice, &PlatformDeviceVersion);
+		FElectraDecodersPlatformResources::GetD3DDeviceAndVersion(&PlatformDevice, &PlatformDeviceVersion);
 		bUseGPUBuffers = (PlatformDevice && PlatformDeviceVersion >= 12000);
 	}
 
@@ -1199,7 +1172,7 @@ bool FElectraVideoDecoderH265_DX::ConvertDecoderOutput()
 		// Create a tasksync instance so we can have our own async jobs run in consecutive order
 		if (!TaskSync.IsValid())
 		{
-			TaskSync = PinnedResourceDelegate->CreateAsyncConsecutiveTaskSync();
+			TaskSync = FElectraDecodersPlatformResources::CreateAsyncConsecutiveTaskSync();
 		}
 
 		// Get data from MFSample...
@@ -1218,10 +1191,10 @@ bool FElectraVideoDecoderH265_DX::ConvertDecoderOutput()
 
 		NewOutput->Pitch = BufferPitch;
 
-		if (PinnedResourceDelegate->RunCodeAsync([DecodedOutputSample, Buffer2D, BufferResource = NewOutput->Buffer.Resource, BufferFence = NewOutput->Buffer.Fence, BufferFenceValue = NewOutput->Buffer.FenceValue, BufferPitch, DecodedHeight = NewOutput->DecodedHeight]()
+		if (FElectraDecodersPlatformResources::RunCodeAsync([DecodedOutputSample, Buffer2D, BufferResource = NewOutput->Buffer.Resource, BufferFence = NewOutput->Buffer.Fence, BufferFenceValue = NewOutput->Buffer.FenceValue, BufferPitch, DecodedHeight = NewOutput->DecodedHeight]()
 			{
 				CopyBufferData(Buffer2D, BufferResource, BufferFence, BufferFenceValue, BufferPitch, DecodedHeight);
-			}, TaskSync.Get()))
+			}, TaskSync))
 		{
 			NewOutput->Buffer_TaskSync = TaskSync;
 		}
diff --git a/Engine/Plugins/Media/ElectraCodecs/Source/ElectraDecoders/Private/DX/VideoDecoderH265_DX.h b/Engine/Plugins/Media/ElectraCodecs/Source/ElectraDecoders/Private/DX/VideoDecoderH265_DX.h
--- a/Engine/Plugins/Media/ElectraCodecs/Source/ElectraDecoders/Private/DX/VideoDecoderH265_DX.h
+++ b/Engine/Plugins/Media/ElectraCodecs/Source/ElectraDecoders/Private/DX/VideoDecoderH265_DX.h
@@ -5,7 +5,6 @@
 #ifdef ELECTRA_DECODERS_ENABLE_DX
 
 #include "IElectraDecoder.h"
-#include "IElectraDecoderResourceDelegate.h"
 
 class IElectraVideoDecoderH265_DX : public IElectraDecoder
 {
@@ -42,11 +41,11 @@ public:
 		virtual void* GetDXDeviceContext() const = 0;
 		virtual bool IsSoftware() const = 0;
 	};
-	static FError PlatformCreateMFDecoderTransform(IPlatformHandle** OutPlatformHandle, TSharedPtr<IElectraDecoderResourceDelegate, ESPMode::ThreadSafe> InResourceDelegate, const TMap<FString, FVariant>& InOptions);
+	static FError PlatformCreateMFDecoderTransform(IPlatformHandle** OutPlatformHandle, const TMap<FString, FVariant>& InOptions);
 	static void PlatformReleaseMFDecoderTransform(IPlatformHandle** OutPlatformHandle);
 
 	
-	static TSharedPtr<IElectraDecoder, ESPMode::ThreadSafe> Create(const TMap<FString, FVariant>& InOptions, TSharedPtr<IElectraDecoderResourceDelegate, ESPMode::ThreadSafe> InResourceDelegate);
+	static TSharedPtr<IElectraDecoder, ESPMode::ThreadSafe> Create(const TMap<FString, FVariant>& InOptions);
 
 	virtual ~IElectraVideoDecoderH265_DX() = default;
 };
diff --git a/Engine/Plugins/Media/ElectraCodecs/Source/ElectraDecoders/Private/ElectraMediaMP3Decoder.cpp b/Engine/Plugins/Media/ElectraCodecs/Source/ElectraDecoders/Private/ElectraMediaMP3Decoder.cpp
--- a/Engine/Plugins/Media/ElectraCodecs/Source/ElectraDecoders/Private/ElectraMediaMP3Decoder.cpp
+++ b/Engine/Plugins/Media/ElectraCodecs/Source/ElectraDecoders/Private/ElectraMediaMP3Decoder.cpp
@@ -11,7 +11,6 @@
 #include "IElectraDecoderFeaturesAndOptions.h"
 #include "IElectraDecoder.h"
 #include "IElectraDecoderOutputAudio.h"
-#include "IElectraDecoderResourceDelegate.h"
 #include "ElectraDecodersUtils.h"
 #include "Utils/ElectraBitstreamReader.h"
 #include "Utils/MPEG/ElectraUtilsMPEGAudio.h"
@@ -35,31 +34,6 @@
 /*********************************************************************************************************************/
 /*********************************************************************************************************************/
 
-class FElectraDecoderDefaultAudioOutputFormatMP3_Common : public IElectraDecoderDefaultAudioOutputFormat
-{
-public:
-	virtual ~FElectraDecoderDefaultAudioOutputFormatMP3_Common()
-	{ }
-
-	int32 GetNumChannels() const override
-	{
-		return NumChannels;
-	}
-	int32 GetSampleRate() const override
-	{
-		return SampleRate;
-	}
-	int32 GetNumFrames() const override
-	{
-		return NumFrames;
-	}
-
-	int32 NumChannels = 0;
-	int32 SampleRate = 0;
-	int32 NumFrames = 0;
-};
-
-
 class FElectraAudioDecoderOutputMP3_Common : public IElectraDecoderAudioOutput
 {
 public:
@@ -132,7 +106,7 @@ public:
 	{
 	}
 
-	FElectraMP3Decoder(const TMap<FString, FVariant>& InOptions, TSharedPtr<IElectraDecoderResourceDelegate, ESPMode::ThreadSafe> InResourceDelegate);
+	FElectraMP3Decoder(const TMap<FString, FVariant>& InOptions);
 
 	virtual ~FElectraMP3Decoder();
 
@@ -149,8 +123,6 @@ public:
 	IElectraDecoder::ECSDCompatibility IsCompatibleWith(const TMap<FString, FVariant>& CSDAndAdditionalOptions) override;
 	bool ResetToCleanStart() override;
 
-	TSharedPtr<IElectraDecoderDefaultOutputFormat, ESPMode::ThreadSafe> GetDefaultOutputFormatFromCSD(const TMap<FString, FVariant>& CSDAndAdditionalOptions) override;
-
 	EDecoderError DecodeAccessUnit(const FInputAccessUnit& InInputAccessUnit, const TMap<FString, FVariant>& InAdditionalOptions) override;
 	EDecoderError SendEndOfData() override;
 	EDecoderError Flush() override;
@@ -248,9 +220,9 @@ public:
 		FElectraMP3Decoder::GetConfigurationOptions(OutOptions);
 	}
 
-	TSharedPtr<IElectraDecoder, ESPMode::ThreadSafe> CreateDecoderForFormat(const FString& InCodecFormat, const TMap<FString, FVariant>& InOptions, TSharedPtr<IElectraDecoderResourceDelegate, ESPMode::ThreadSafe> InResourceDelegate) override
+	TSharedPtr<IElectraDecoder, ESPMode::ThreadSafe> CreateDecoderForFormat(const FString& InCodecFormat, const TMap<FString, FVariant>& InOptions) override
 	{
-		return MakeShared<FElectraMP3Decoder, ESPMode::ThreadSafe>(InOptions, InResourceDelegate);
+		return MakeShared<FElectraMP3Decoder, ESPMode::ThreadSafe>(InOptions);
 	}
 
 	static TSharedPtr<FElectraCommonAudioMP3DecoderFactory, ESPMode::ThreadSafe> Self;
@@ -294,7 +266,7 @@ TSharedPtr<IElectraCodecFactory, ESPMode::ThreadSafe> FElectraMediaMP3Decoder::C
 /*********************************************************************************************************************/
 /*********************************************************************************************************************/
 
-FElectraMP3Decoder::FElectraMP3Decoder(const TMap<FString, FVariant>& InOptions, TSharedPtr<IElectraDecoderResourceDelegate, ESPMode::ThreadSafe> InResourceDelegate)
+FElectraMP3Decoder::FElectraMP3Decoder(const TMap<FString, FVariant>& InOptions)
 {
 	Codec4CC = (uint32)ElectraDecodersUtil::GetVariantValueSafeU64(InOptions, TEXT("codec_4cc"), 0);
 }
@@ -400,11 +372,6 @@ bool FElectraMP3Decoder::ResetToCleanStart()
 	return true;
 }
 
-TSharedPtr<IElectraDecoderDefaultOutputFormat, ESPMode::ThreadSafe> FElectraMP3Decoder::GetDefaultOutputFormatFromCSD(const TMap<FString, FVariant>& CSDAndAdditionalOptions)
-{
-	return nullptr;
-}
-
 bool FElectraMP3Decoder::SetupChannelMap()
 {
 	if (OutputChannelMap.Num())
diff --git a/Engine/Plugins/Media/ElectraCodecs/Source/ElectraDecoders/Private/Linux/LinuxElectraDecodersPlatformResources.cpp b/Engine/Plugins/Media/ElectraCodecs/Source/ElectraDecoders/Private/Linux/LinuxElectraDecodersPlatformResources.cpp
new file mode 100644
--- /dev/null
+++ b/Engine/Plugins/Media/ElectraCodecs/Source/ElectraDecoders/Private/Linux/LinuxElectraDecodersPlatformResources.cpp
@@ -0,0 +1,10 @@
+// Copyright Epic Games, Inc. All Rights Reserved.
+#include "Linux/LinuxElectraDecodersPlatformResources.h"
+
+
+/*
+bool FElectraDecodersPlatformResourcesLinux::NotUsedRightNow()
+{
+	return false;
+}
+*/
\ No newline at end of file
diff --git a/Engine/Plugins/Media/ElectraCodecs/Source/ElectraDecoders/Private/Linux/aac/AAC_AudioDecoder_Linux.cpp b/Engine/Plugins/Media/ElectraCodecs/Source/ElectraDecoders/Private/Linux/aac/AAC_AudioDecoder_Linux.cpp
--- a/Engine/Plugins/Media/ElectraCodecs/Source/ElectraDecoders/Private/Linux/aac/AAC_AudioDecoder_Linux.cpp
+++ b/Engine/Plugins/Media/ElectraCodecs/Source/ElectraDecoders/Private/Linux/aac/AAC_AudioDecoder_Linux.cpp
@@ -80,9 +80,9 @@ public:
 		IElectraAudioDecoderAAC_Linux::GetConfigurationOptions(OutOptions);
 	}
 
-	TSharedPtr<IElectraDecoder, ESPMode::ThreadSafe> CreateDecoderForFormat(const FString& InCodecFormat, const TMap<FString, FVariant>& InOptions, TSharedPtr<IElectraDecoderResourceDelegate, ESPMode::ThreadSafe> InResourceDelegate) override
+	TSharedPtr<IElectraDecoder, ESPMode::ThreadSafe> CreateDecoderForFormat(const FString& InCodecFormat, const TMap<FString, FVariant>& InOptions) override
 	{
-		return IElectraAudioDecoderAAC_Linux::Create(InOptions, InResourceDelegate);
+		return IElectraAudioDecoderAAC_Linux::Create(InOptions);
 	}
 
 };
diff --git a/Engine/Plugins/Media/ElectraCodecs/Source/ElectraDecoders/Private/Linux/aac/AudioDecoderAAC_Linux.cpp b/Engine/Plugins/Media/ElectraCodecs/Source/ElectraDecoders/Private/Linux/aac/AudioDecoderAAC_Linux.cpp
--- a/Engine/Plugins/Media/ElectraCodecs/Source/ElectraDecoders/Private/Linux/aac/AudioDecoderAAC_Linux.cpp
+++ b/Engine/Plugins/Media/ElectraCodecs/Source/ElectraDecoders/Private/Linux/aac/AudioDecoderAAC_Linux.cpp
@@ -12,8 +12,6 @@
 #include "IElectraDecoderFeaturesAndOptions.h"
 #include "IElectraDecoderOutputAudio.h"
 
-#include "IElectraDecoderResourceDelegate.h"
-
 /*********************************************************************************************************************/
 
 #include "libav_Decoder_Common.h"
@@ -21,25 +19,6 @@
 
 /*********************************************************************************************************************/
 
-class FElectraDecoderDefaultAudioOutputFormatAAC_Linux : public IElectraDecoderDefaultAudioOutputFormat
-{
-public:
-	virtual ~FElectraDecoderDefaultAudioOutputFormatAAC_Linux()
-	{ }
-
-	int32 GetNumChannels() const override
-	{ return NumChannels; }
-	int32 GetSampleRate() const override
-	{ return SampleRate; }
-	int32 GetNumFrames() const override
-	{ return NumFrames; }
-
-	int32 NumChannels = 0;
-	int32 SampleRate = 0;
-	int32 NumFrames = 0;
-};
-
-
 class FElectraAudioDecoderOutputAAC_Linux : public IElectraDecoderAudioOutput
 {
 public:
@@ -91,7 +70,7 @@ public:
 class FElectraAudioDecoderAAC_Linux : public IElectraAudioDecoderAAC_Linux
 {
 public:
-	FElectraAudioDecoderAAC_Linux(const TMap<FString, FVariant>& InOptions, TSharedPtr<IElectraDecoderResourceDelegate, ESPMode::ThreadSafe> InResourceDelegate);
+	FElectraAudioDecoderAAC_Linux(const TMap<FString, FVariant>& InOptions);
 
 	virtual ~FElectraAudioDecoderAAC_Linux();
 
@@ -106,8 +85,6 @@ public:
 	IElectraDecoder::ECSDCompatibility IsCompatibleWith(const TMap<FString, FVariant>& CSDAndAdditionalOptions) override;
 	bool ResetToCleanStart() override;
 
-	TSharedPtr<IElectraDecoderDefaultOutputFormat, ESPMode::ThreadSafe> GetDefaultOutputFormatFromCSD(const TMap<FString, FVariant>& CSDAndAdditionalOptions) override;
-
 	EDecoderError DecodeAccessUnit(const FInputAccessUnit& InInputAccessUnit, const TMap<FString, FVariant>& InAdditionalOptions) override;
 	EDecoderError SendEndOfData() override;
 	EDecoderError Flush() override;
@@ -142,8 +119,6 @@ private:
 
 	bool ConvertDecoderOutput(const ILibavDecoderAAC::FOutputInfo& OutInfo);
 
-	TWeakPtr<IElectraDecoderResourceDelegate, ESPMode::ThreadSafe> ResourceDelegate;
-
 	IElectraDecoder::FError LastError;
 
 	TSharedPtr<ElectraDecodersUtil::MPEG::FAACDecoderConfigurationRecord, ESPMode::ThreadSafe> ConfigRecord;
@@ -180,20 +155,14 @@ void IElectraAudioDecoderAAC_Linux::GetConfigurationOptions(TMap<FString, FVaria
 {
 }
 
-TSharedPtr<IElectraDecoder, ESPMode::ThreadSafe> IElectraAudioDecoderAAC_Linux::Create(const TMap<FString, FVariant>& InOptions, TSharedPtr<IElectraDecoderResourceDelegate, ESPMode::ThreadSafe> InResourceDelegate)
+TSharedPtr<IElectraDecoder, ESPMode::ThreadSafe> IElectraAudioDecoderAAC_Linux::Create(const TMap<FString, FVariant>& InOptions)
 {
-	if (InResourceDelegate)
-	{
-		TSharedPtr<FElectraAudioDecoderAAC_Linux, ESPMode::ThreadSafe> New = MakeShared<FElectraAudioDecoderAAC_Linux>(InOptions, InResourceDelegate);
-		return New;
-	}
-	return nullptr;
+	TSharedPtr<FElectraAudioDecoderAAC_Linux, ESPMode::ThreadSafe> New = MakeShared<FElectraAudioDecoderAAC_Linux>(InOptions);
+	return New;
 }
 
-FElectraAudioDecoderAAC_Linux::FElectraAudioDecoderAAC_Linux(const TMap<FString, FVariant>& InOptions, TSharedPtr<IElectraDecoderResourceDelegate, ESPMode::ThreadSafe> InResourceDelegate)
+FElectraAudioDecoderAAC_Linux::FElectraAudioDecoderAAC_Linux(const TMap<FString, FVariant>& InOptions)
 {
-	ResourceDelegate = InResourceDelegate;
-
 	// If there is codec specific data set we create a config record from it now.
 	ConfigRecord = GetCSDFromOptions(InOptions);
 }
@@ -257,20 +226,6 @@ bool FElectraAudioDecoderAAC_Linux::ResetToCleanStart()
 	return !LastError.IsSet();
 }
 
-TSharedPtr<IElectraDecoderDefaultOutputFormat, ESPMode::ThreadSafe> FElectraAudioDecoderAAC_Linux::GetDefaultOutputFormatFromCSD(const TMap<FString, FVariant>& CSDAndAdditionalOptions)
-{
-	TSharedPtr<ElectraDecodersUtil::MPEG::FAACDecoderConfigurationRecord, ESPMode::ThreadSafe> Cfg = GetCSDFromOptions(CSDAndAdditionalOptions);
-	if (Cfg.IsValid())
-	{
-		TSharedPtr<FElectraDecoderDefaultAudioOutputFormatAAC_Linux, ESPMode::ThreadSafe> Info = MakeShared<FElectraDecoderDefaultAudioOutputFormatAAC_Linux, ESPMode::ThreadSafe>();
-		Info->NumChannels = Cfg->PSSignal > 0 ? 2 : NumChannelsForConfig[Cfg->ChannelConfiguration];
-		Info->SampleRate = Cfg->ExtSamplingFrequency ? Cfg->ExtSamplingFrequency : ConfigRecord->SamplingRate;
-		Info->NumFrames = Cfg->SBRSignal > 0 ? 2048 : 1024;
-		return Info;
-	}
-	return nullptr;
-}
-
 IElectraDecoder::EDecoderError FElectraAudioDecoderAAC_Linux::DecodeAccessUnit(const FInputAccessUnit& InInputAccessUnit, const TMap<FString, FVariant>& InAdditionalOptions)
 {
 	// If already in error do nothing!
diff --git a/Engine/Plugins/Media/ElectraCodecs/Source/ElectraDecoders/Private/Linux/h264/H264_VideoDecoder_Linux.cpp b/Engine/Plugins/Media/ElectraCodecs/Source/ElectraDecoders/Private/Linux/h264/H264_VideoDecoder_Linux.cpp
--- a/Engine/Plugins/Media/ElectraCodecs/Source/ElectraDecoders/Private/Linux/h264/H264_VideoDecoder_Linux.cpp
+++ b/Engine/Plugins/Media/ElectraCodecs/Source/ElectraDecoders/Private/Linux/h264/H264_VideoDecoder_Linux.cpp
@@ -91,9 +91,9 @@ public:
 		IElectraVideoDecoderH264_Linux::GetConfigurationOptions(OutOptions);
 	}
 
-	TSharedPtr<IElectraDecoder, ESPMode::ThreadSafe> CreateDecoderForFormat(const FString& InCodecFormat, const TMap<FString, FVariant>& InOptions, TSharedPtr<IElectraDecoderResourceDelegate, ESPMode::ThreadSafe> InResourceDelegate) override
+	TSharedPtr<IElectraDecoder, ESPMode::ThreadSafe> CreateDecoderForFormat(const FString& InCodecFormat, const TMap<FString, FVariant>& InOptions) override
 	{
-		return IElectraVideoDecoderH264_Linux::Create(InOptions, InResourceDelegate);
+		return IElectraVideoDecoderH264_Linux::Create(InOptions);
 	}
 
 };
diff --git a/Engine/Plugins/Media/ElectraCodecs/Source/ElectraDecoders/Private/Linux/h264/VideoDecoderH264_Linux.cpp b/Engine/Plugins/Media/ElectraCodecs/Source/ElectraDecoders/Private/Linux/h264/VideoDecoderH264_Linux.cpp
--- a/Engine/Plugins/Media/ElectraCodecs/Source/ElectraDecoders/Private/Linux/h264/VideoDecoderH264_Linux.cpp
+++ b/Engine/Plugins/Media/ElectraCodecs/Source/ElectraDecoders/Private/Linux/h264/VideoDecoderH264_Linux.cpp
@@ -9,7 +9,6 @@
 #include "IElectraDecoderFeaturesAndOptions.h"
 #include "IElectraDecoderOutputVideo.h"
 
-#include "IElectraDecoderResourceDelegate.h"
 #include "ElectraDecodersModule.h"
 
 /*********************************************************************************************************************/
@@ -22,15 +21,6 @@
 class FElectraVideoDecoderH264_Linux;
 
 
-class FElectraDecoderDefaultVideoOutputFormatH264_Linux : public IElectraDecoderDefaultVideoOutputFormat
-{
-public:
-	virtual ~FElectraDecoderDefaultVideoOutputFormatH264_Linux()
-	{ }
-
-};
-
-
 class FElectraVideoDecoderOutputH264_Linux : public IElectraDecoderVideoOutput
 {
 public:
@@ -74,8 +64,6 @@ public:
 		}
 		return nullptr;
 	}
-	IElectraDecoderVideoOutputTransferHandle* GetTransferHandle() const override
-	{ return nullptr; }
 	IElectraDecoderVideoOutput::EImageCopyResult CopyPlatformImage(IElectraDecoderVideoOutputCopyResources* InCopyResources) const override
 	{ return IElectraDecoderVideoOutput::EImageCopyResult::NotSupported; }
 
@@ -103,7 +91,7 @@ public:
 class FElectraVideoDecoderH264_Linux : public IElectraVideoDecoderH264_Linux, public ILibavDecoderVideoResourceAllocator
 {
 public:
-	FElectraVideoDecoderH264_Linux(const TMap<FString, FVariant>& InOptions, TSharedPtr<IElectraDecoderResourceDelegate, ESPMode::ThreadSafe> InResourceDelegate);
+	FElectraVideoDecoderH264_Linux(const TMap<FString, FVariant>& InOptions);
 
 	virtual ~FElectraVideoDecoderH264_Linux();
 
@@ -118,8 +106,6 @@ public:
 	IElectraDecoder::ECSDCompatibility IsCompatibleWith(const TMap<FString, FVariant>& CSDAndAdditionalOptions) override;
 	bool ResetToCleanStart() override;
 
-	TSharedPtr<IElectraDecoderDefaultOutputFormat, ESPMode::ThreadSafe> GetDefaultOutputFormatFromCSD(const TMap<FString, FVariant>& CSDAndAdditionalOptions) override;
-
 	EDecoderError DecodeAccessUnit(const FInputAccessUnit& InInputAccessUnit, const TMap<FString, FVariant>& InAdditionalOptions) override;
 	EDecoderError SendEndOfData() override;
 	EDecoderError Flush() override;
@@ -169,7 +155,6 @@ private:
 	EConvertResult ConvertDecoderOutput(const ILibavDecoderVideoCommon::FOutputInfo& OutInfo);
 
 	TMap<FString, FVariant> InitialCreationOptions;
-	TWeakPtr<IElectraDecoderResourceDelegate, ESPMode::ThreadSafe> ResourceDelegate;
 
 	IElectraDecoder::FError LastError;
 
@@ -217,20 +202,15 @@ void IElectraVideoDecoderH264_Linux::GetConfigurationOptions(TMap<FString, FVari
 	OutOptions.Emplace(IElectraDecoderFeature::SupportsDroppingOutput, FVariant(true));
 }
 
-TSharedPtr<IElectraDecoder, ESPMode::ThreadSafe> IElectraVideoDecoderH264_Linux::Create(const TMap<FString, FVariant>& InOptions, TSharedPtr<IElectraDecoderResourceDelegate, ESPMode::ThreadSafe> InResourceDelegate)
+TSharedPtr<IElectraDecoder, ESPMode::ThreadSafe> IElectraVideoDecoderH264_Linux::Create(const TMap<FString, FVariant>& InOptions)
 {
-	if (InResourceDelegate.IsValid())
-	{
-		TSharedPtr<FElectraVideoDecoderH264_Linux, ESPMode::ThreadSafe> New = MakeShared<FElectraVideoDecoderH264_Linux>(InOptions, InResourceDelegate);
-		return New;
-	}
-	return nullptr;
+	TSharedPtr<FElectraVideoDecoderH264_Linux, ESPMode::ThreadSafe> New = MakeShared<FElectraVideoDecoderH264_Linux>(InOptions);
+	return New;
 }
 
-FElectraVideoDecoderH264_Linux::FElectraVideoDecoderH264_Linux(const TMap<FString, FVariant>& InOptions, TSharedPtr<IElectraDecoderResourceDelegate, ESPMode::ThreadSafe> InResourceDelegate)
+FElectraVideoDecoderH264_Linux::FElectraVideoDecoderH264_Linux(const TMap<FString, FVariant>& InOptions)
 {
 	InitialCreationOptions = InOptions;
-	ResourceDelegate = InResourceDelegate;
 }
 
 FElectraVideoDecoderH264_Linux::~FElectraVideoDecoderH264_Linux()
@@ -298,11 +278,6 @@ bool FElectraVideoDecoderH264_Linux::ResetToCleanStart()
 	return !LastError.IsSet();
 }
 
-TSharedPtr<IElectraDecoderDefaultOutputFormat, ESPMode::ThreadSafe> FElectraVideoDecoderH264_Linux::GetDefaultOutputFormatFromCSD(const TMap<FString, FVariant>& CSDAndAdditionalOptions)
-{
-	return nullptr;
-}
-
 IElectraDecoder::ECSDCompatibility FElectraVideoDecoderH264_Linux::IsCompatibleWith(const TMap<FString, FVariant>& CSDAndAdditionalOptions)
 {
 	// NOTE: Since we set IElectraDecoderFeature::IsAdaptive to be true this method here will not be invoked on stream switches.
diff --git a/Engine/Plugins/Media/ElectraCodecs/Source/ElectraDecoders/Private/Linux/h265/H265_VideoDecoder_Linux.cpp b/Engine/Plugins/Media/ElectraCodecs/Source/ElectraDecoders/Private/Linux/h265/H265_VideoDecoder_Linux.cpp
--- a/Engine/Plugins/Media/ElectraCodecs/Source/ElectraDecoders/Private/Linux/h265/H265_VideoDecoder_Linux.cpp
+++ b/Engine/Plugins/Media/ElectraCodecs/Source/ElectraDecoders/Private/Linux/h265/H265_VideoDecoder_Linux.cpp
@@ -91,9 +91,9 @@ public:
 		IElectraVideoDecoderH265_Linux::GetConfigurationOptions(OutOptions);
 	}
 
-	TSharedPtr<IElectraDecoder, ESPMode::ThreadSafe> CreateDecoderForFormat(const FString& InCodecFormat, const TMap<FString, FVariant>& InOptions, TSharedPtr<IElectraDecoderResourceDelegate, ESPMode::ThreadSafe> InResourceDelegate) override
+	TSharedPtr<IElectraDecoder, ESPMode::ThreadSafe> CreateDecoderForFormat(const FString& InCodecFormat, const TMap<FString, FVariant>& InOptions) override
 	{
-		return IElectraVideoDecoderH265_Linux::Create(InOptions, InResourceDelegate);
+		return IElectraVideoDecoderH265_Linux::Create(InOptions);
 	}
 
 };
diff --git a/Engine/Plugins/Media/ElectraCodecs/Source/ElectraDecoders/Private/Linux/h265/VideoDecoderH265_Linux.cpp b/Engine/Plugins/Media/ElectraCodecs/Source/ElectraDecoders/Private/Linux/h265/VideoDecoderH265_Linux.cpp
--- a/Engine/Plugins/Media/ElectraCodecs/Source/ElectraDecoders/Private/Linux/h265/VideoDecoderH265_Linux.cpp
+++ b/Engine/Plugins/Media/ElectraCodecs/Source/ElectraDecoders/Private/Linux/h265/VideoDecoderH265_Linux.cpp
@@ -9,7 +9,6 @@
 #include "IElectraDecoderFeaturesAndOptions.h"
 #include "IElectraDecoderOutputVideo.h"
 
-#include "IElectraDecoderResourceDelegate.h"
 #include "ElectraDecodersModule.h"
 
 /*********************************************************************************************************************/
@@ -22,15 +21,6 @@
 class FElectraVideoDecoderH265_Linux;
 
 
-class FElectraDecoderDefaultVideoOutputFormatH265_Linux : public IElectraDecoderDefaultVideoOutputFormat
-{
-public:
-	virtual ~FElectraDecoderDefaultVideoOutputFormatH265_Linux()
-	{ }
-
-};
-
-
 class FElectraVideoDecoderOutputH265_Linux : public IElectraDecoderVideoOutput
 {
 public:
@@ -74,8 +64,6 @@ public:
 		}
 		return nullptr;
 	}
-	IElectraDecoderVideoOutputTransferHandle* GetTransferHandle() const override
-	{ return nullptr; }
 	IElectraDecoderVideoOutput::EImageCopyResult CopyPlatformImage(IElectraDecoderVideoOutputCopyResources* InCopyResources) const override
 	{ return IElectraDecoderVideoOutput::EImageCopyResult::NotSupported; }
 
@@ -103,7 +91,7 @@ public:
 class FElectraVideoDecoderH265_Linux : public IElectraVideoDecoderH265_Linux, public ILibavDecoderVideoResourceAllocator
 {
 public:
-	FElectraVideoDecoderH265_Linux(const TMap<FString, FVariant>& InOptions, TSharedPtr<IElectraDecoderResourceDelegate, ESPMode::ThreadSafe> InResourceDelegate);
+	FElectraVideoDecoderH265_Linux(const TMap<FString, FVariant>& InOptions);
 
 	virtual ~FElectraVideoDecoderH265_Linux();
 
@@ -118,8 +106,6 @@ public:
 	IElectraDecoder::ECSDCompatibility IsCompatibleWith(const TMap<FString, FVariant>& CSDAndAdditionalOptions) override;
 	bool ResetToCleanStart() override;
 
-	TSharedPtr<IElectraDecoderDefaultOutputFormat, ESPMode::ThreadSafe> GetDefaultOutputFormatFromCSD(const TMap<FString, FVariant>& CSDAndAdditionalOptions) override;
-
 	EDecoderError DecodeAccessUnit(const FInputAccessUnit& InInputAccessUnit, const TMap<FString, FVariant>& InAdditionalOptions) override;
 	EDecoderError SendEndOfData() override;
 	EDecoderError Flush() override;
@@ -170,7 +156,6 @@ private:
 	EConvertResult ConvertDecoderOutput(const ILibavDecoderVideoCommon::FOutputInfo& OutInfo);
 
 	TMap<FString, FVariant> InitialCreationOptions;
-	TWeakPtr<IElectraDecoderResourceDelegate, ESPMode::ThreadSafe> ResourceDelegate;
 
 	IElectraDecoder::FError LastError;
 
@@ -219,20 +204,15 @@ void IElectraVideoDecoderH265_Linux::GetConfigurationOptions(TMap<FString, FVari
 	OutOptions.Emplace(IElectraDecoderFeature::SupportsDroppingOutput, FVariant(true));
 }
 
-TSharedPtr<IElectraDecoder, ESPMode::ThreadSafe> IElectraVideoDecoderH265_Linux::Create(const TMap<FString, FVariant>& InOptions, TSharedPtr<IElectraDecoderResourceDelegate, ESPMode::ThreadSafe> InResourceDelegate)
+TSharedPtr<IElectraDecoder, ESPMode::ThreadSafe> IElectraVideoDecoderH265_Linux::Create(const TMap<FString, FVariant>& InOptions)
 {
-	if (InResourceDelegate.IsValid())
-	{
-		TSharedPtr<FElectraVideoDecoderH265_Linux, ESPMode::ThreadSafe> New = MakeShared<FElectraVideoDecoderH265_Linux>(InOptions, InResourceDelegate);
-		return New;
-	}
-	return nullptr;
+	TSharedPtr<FElectraVideoDecoderH265_Linux, ESPMode::ThreadSafe> New = MakeShared<FElectraVideoDecoderH265_Linux>(InOptions);
+	return New;
 }
 
-FElectraVideoDecoderH265_Linux::FElectraVideoDecoderH265_Linux(const TMap<FString, FVariant>& InOptions, TSharedPtr<IElectraDecoderResourceDelegate, ESPMode::ThreadSafe> InResourceDelegate)
+FElectraVideoDecoderH265_Linux::FElectraVideoDecoderH265_Linux(const TMap<FString, FVariant>& InOptions)
 {
 	InitialCreationOptions = InOptions;
-	ResourceDelegate = InResourceDelegate;
 }
 
 FElectraVideoDecoderH265_Linux::~FElectraVideoDecoderH265_Linux()
@@ -319,11 +299,6 @@ bool FElectraVideoDecoderH265_Linux::ResetToCleanStart()
 	return !LastError.IsSet();
 }
 
-TSharedPtr<IElectraDecoderDefaultOutputFormat, ESPMode::ThreadSafe> FElectraVideoDecoderH265_Linux::GetDefaultOutputFormatFromCSD(const TMap<FString, FVariant>& CSDAndAdditionalOptions)
-{
-	return nullptr;
-}
-
 IElectraDecoder::ECSDCompatibility FElectraVideoDecoderH265_Linux::IsCompatibleWith(const TMap<FString, FVariant>& CSDAndAdditionalOptions)
 {
 	// NOTE: Since we set IElectraDecoderFeature::IsAdaptive to be true this method here will not be invoked on stream switches.
diff --git a/Engine/Plugins/Media/ElectraCodecs/Source/ElectraDecoders/Private/Linux/vpx/VPx_VideoDecoder_Linux.cpp b/Engine/Plugins/Media/ElectraCodecs/Source/ElectraDecoders/Private/Linux/vpx/VPx_VideoDecoder_Linux.cpp
--- a/Engine/Plugins/Media/ElectraCodecs/Source/ElectraDecoders/Private/Linux/vpx/VPx_VideoDecoder_Linux.cpp
+++ b/Engine/Plugins/Media/ElectraCodecs/Source/ElectraDecoders/Private/Linux/vpx/VPx_VideoDecoder_Linux.cpp
@@ -79,9 +79,9 @@ public:
 		IElectraVideoDecoderVPx_Linux::GetConfigurationOptions(OutOptions);
 	}
 
-	TSharedPtr<IElectraDecoder, ESPMode::ThreadSafe> CreateDecoderForFormat(const FString& InCodecFormat, const TMap<FString, FVariant>& InOptions, TSharedPtr<IElectraDecoderResourceDelegate, ESPMode::ThreadSafe> InResourceDelegate) override
+	TSharedPtr<IElectraDecoder, ESPMode::ThreadSafe> CreateDecoderForFormat(const FString& InCodecFormat, const TMap<FString, FVariant>& InOptions) override
 	{
-		return IElectraVideoDecoderVPx_Linux::Create(InOptions, InResourceDelegate);
+		return IElectraVideoDecoderVPx_Linux::Create(InOptions);
 	}
 
 };
diff --git a/Engine/Plugins/Media/ElectraCodecs/Source/ElectraDecoders/Private/Linux/vpx/VideoDecoderVPx_Linux.cpp b/Engine/Plugins/Media/ElectraCodecs/Source/ElectraDecoders/Private/Linux/vpx/VideoDecoderVPx_Linux.cpp
--- a/Engine/Plugins/Media/ElectraCodecs/Source/ElectraDecoders/Private/Linux/vpx/VideoDecoderVPx_Linux.cpp
+++ b/Engine/Plugins/Media/ElectraCodecs/Source/ElectraDecoders/Private/Linux/vpx/VideoDecoderVPx_Linux.cpp
@@ -9,7 +9,6 @@
 #include "IElectraDecoderFeaturesAndOptions.h"
 #include "IElectraDecoderOutputVideo.h"
 
-#include "IElectraDecoderResourceDelegate.h"
 #include "ElectraDecodersModule.h"
 
 /*********************************************************************************************************************/
@@ -23,15 +22,6 @@
 class FElectraVideoDecoderVPx_Linux;
 
 
-class FElectraDecoderDefaultVideoOutputFormatVPx_Linux : public IElectraDecoderDefaultVideoOutputFormat
-{
-public:
-	virtual ~FElectraDecoderDefaultVideoOutputFormatVPx_Linux()
-	{ }
-
-};
-
-
 class FElectraVideoDecoderOutputVPx_Linux : public IElectraDecoderVideoOutput
 {
 public:
@@ -75,8 +65,6 @@ public:
 		}
 		return nullptr;
 	}
-	IElectraDecoderVideoOutputTransferHandle* GetTransferHandle() const override
-	{ return nullptr; }
 	IElectraDecoderVideoOutput::EImageCopyResult CopyPlatformImage(IElectraDecoderVideoOutputCopyResources* InCopyResources) const override
 	{ return IElectraDecoderVideoOutput::EImageCopyResult::NotSupported; }
 
@@ -104,7 +92,7 @@ public:
 class FElectraVideoDecoderVPx_Linux : public IElectraVideoDecoderVPx_Linux, public ILibavDecoderVideoResourceAllocator
 {
 public:
-	FElectraVideoDecoderVPx_Linux(const TMap<FString, FVariant>& InOptions, TSharedPtr<IElectraDecoderResourceDelegate, ESPMode::ThreadSafe> InResourceDelegate);
+	FElectraVideoDecoderVPx_Linux(const TMap<FString, FVariant>& InOptions);
 
 	virtual ~FElectraVideoDecoderVPx_Linux();
 
@@ -119,8 +107,6 @@ public:
 	IElectraDecoder::ECSDCompatibility IsCompatibleWith(const TMap<FString, FVariant>& CSDAndAdditionalOptions) override;
 	bool ResetToCleanStart() override;
 
-	TSharedPtr<IElectraDecoderDefaultOutputFormat, ESPMode::ThreadSafe> GetDefaultOutputFormatFromCSD(const TMap<FString, FVariant>& CSDAndAdditionalOptions) override;
-
 	EDecoderError DecodeAccessUnit(const FInputAccessUnit& InInputAccessUnit, const TMap<FString, FVariant>& InAdditionalOptions) override;
 	EDecoderError SendEndOfData() override;
 	EDecoderError Flush() override;
@@ -175,7 +161,6 @@ private:
 
 	ECodec ActiveCodec = ECodec::VP9;
 	TMap<FString, FVariant> InitialCreationOptions;
-	TWeakPtr<IElectraDecoderResourceDelegate, ESPMode::ThreadSafe> ResourceDelegate;
 
 	IElectraDecoder::FError LastError;
 
@@ -196,20 +181,15 @@ void IElectraVideoDecoderVPx_Linux::GetConfigurationOptions(TMap<FString, FVaria
 	OutOptions.Emplace(IElectraDecoderFeature::SupportsDroppingOutput, FVariant(true));
 }
 
-TSharedPtr<IElectraDecoder, ESPMode::ThreadSafe> IElectraVideoDecoderVPx_Linux::Create(const TMap<FString, FVariant>& InOptions, TSharedPtr<IElectraDecoderResourceDelegate, ESPMode::ThreadSafe> InResourceDelegate)
+TSharedPtr<IElectraDecoder, ESPMode::ThreadSafe> IElectraVideoDecoderVPx_Linux::Create(const TMap<FString, FVariant>& InOptions)
 {
-	if (InResourceDelegate.IsValid())
-	{
-		TSharedPtr<FElectraVideoDecoderVPx_Linux, ESPMode::ThreadSafe> New = MakeShared<FElectraVideoDecoderVPx_Linux>(InOptions, InResourceDelegate);
-		return New;
-	}
-	return nullptr;
+	TSharedPtr<FElectraVideoDecoderVPx_Linux, ESPMode::ThreadSafe> New = MakeShared<FElectraVideoDecoderVPx_Linux>(InOptions);
+	return New;
 }
 
-FElectraVideoDecoderVPx_Linux::FElectraVideoDecoderVPx_Linux(const TMap<FString, FVariant>& InOptions, TSharedPtr<IElectraDecoderResourceDelegate, ESPMode::ThreadSafe> InResourceDelegate)
+FElectraVideoDecoderVPx_Linux::FElectraVideoDecoderVPx_Linux(const TMap<FString, FVariant>& InOptions)
 {
 	InitialCreationOptions = InOptions;
-	ResourceDelegate = InResourceDelegate;
 
 	uint32 Codec4CC = (uint32) ElectraDecodersUtil::GetVariantValueSafeI64(InOptions, TEXT("codec_4cc"), 0);
 	ActiveCodec = Codec4CC == 'vp09' ? ECodec::VP9 : ECodec::VP8;
@@ -251,11 +231,6 @@ bool FElectraVideoDecoderVPx_Linux::ResetToCleanStart()
 	return !LastError.IsSet();
 }
 
-TSharedPtr<IElectraDecoderDefaultOutputFormat, ESPMode::ThreadSafe> FElectraVideoDecoderVPx_Linux::GetDefaultOutputFormatFromCSD(const TMap<FString, FVariant>& CSDAndAdditionalOptions)
-{
-	return nullptr;
-}
-
 IElectraDecoder::ECSDCompatibility FElectraVideoDecoderVPx_Linux::IsCompatibleWith(const TMap<FString, FVariant>& CSDAndAdditionalOptions)
 {
 	// NOTE: Since we set IElectraDecoderFeature::IsAdaptive to be true this method here will not be invoked on stream switches.
diff --git a/Engine/Plugins/Media/ElectraCodecs/Source/ElectraDecoders/Private/Windows/WindowsPlatformElectraDecoders.cpp b/Engine/Plugins/Media/ElectraCodecs/Source/ElectraDecoders/Private/Windows/WindowsPlatformElectraDecoders.cpp
--- a/Engine/Plugins/Media/ElectraCodecs/Source/ElectraDecoders/Private/Windows/WindowsPlatformElectraDecoders.cpp
+++ b/Engine/Plugins/Media/ElectraCodecs/Source/ElectraDecoders/Private/Windows/WindowsPlatformElectraDecoders.cpp
@@ -1,25 +1,82 @@
 // Copyright Epic Games, Inc. All Rights Reserved.
 
 #include "WindowsPlatformElectraDecoders.h"
+#include "WindowsElectraDecodersPlatformResources.h"
 #include "IElectraCodecRegistry.h"
+#include "ElectraDecodersModule.h"
 
 #include "h264/H264_VideoDecoder_Windows.h"
 #include "h265/H265_VideoDecoder_Windows.h"
 #include "aac/AAC_AudioDecoder_Windows.h"
 
+#include "HAL/PlatformProcess.h"
+#include "RHI.h"
+#include "Async/Async.h"
+
+#include "Microsoft/AllowMicrosoftPlatformTypes.h"
+THIRD_PARTY_INCLUDES_START
+#include "mftransform.h"
+#include "mfapi.h"
+#include "mferror.h"
+#include "mfidl.h"
+THIRD_PARTY_INCLUDES_END
+#include "Microsoft/HideMicrosoftPlatformTypes.h"
+
+DECLARE_CYCLE_STAT(TEXT("Electra decoder AsyncJob"), STAT_ElectraDecoderAsyncJob, STATGROUP_Media);
+
+
 namespace PlatformElectraDecodersWindows
 {
 	static TArray<TSharedPtr<IElectraCodecFactory, ESPMode::ThreadSafe>> Factories;
+	static bool bMFInitialized = false;
+
+	static FCriticalSection AsyncJobAccessCS;
+	static FGraphEventRef RunCodeAsyncEvent;
 }
 
 
 void FPlatformElectraDecodersWindows::Startup()
 {
+	// Minimum of Windows 8 required.
+	if (!FPlatformMisc::VerifyWindowsVersion(6, 2))
+	{
+		UE_LOG(LogElectraDecoders, Error, TEXT("Need at least Windows 8.0."));
+		return;
+	}
+	// We need RHI to perform conversion from the decoder output.
+	if (GDynamicRHI == nullptr || RHIGetInterfaceType() == ERHIInterfaceType::Null)
+	{
+		UE_LOG(LogElectraDecoders, Log, TEXT("No or only dummy dynamic RHI detected, cannot use Electra decoders without RHI."));
+		return;
+	}
+
+	// Prepare Microsoft's Media Foundation
+	if (!(FPlatformProcess::GetDllHandle(TEXT("mf.dll"))
+		&& FPlatformProcess::GetDllHandle(TEXT("mfplat.dll"))
+		&& FPlatformProcess::GetDllHandle(TEXT("msmpeg2vdec.dll"))
+		&& FPlatformProcess::GetDllHandle(TEXT("MSAudDecMFT.dll"))))
+	{
+		UE_LOG(LogElectraDecoders, Error, TEXT("Could not load Media Foundation"));
+		return;
+	}
+	HRESULT Result = MFStartup(MF_VERSION);
+	if (FAILED(Result))
+	{
+		UE_LOG(LogElectraDecoders, Error, TEXT("MFStartup failed with 0x%08x"), Result);
+		return;
+	}
+	PlatformElectraDecodersWindows::bMFInitialized = true;
 }
 
 void FPlatformElectraDecodersWindows::Shutdown()
 {
 	PlatformElectraDecodersWindows::Factories.Empty();
+
+	if (PlatformElectraDecodersWindows::bMFInitialized)
+	{
+		PlatformElectraDecodersWindows::bMFInitialized = false;
+		MFShutdown();
+	}
 }
 
 void FPlatformElectraDecodersWindows::RegisterWithCodecFactory(IElectraCodecRegistry* InCodecFactoryToRegisterWith)
@@ -46,3 +103,52 @@ void FPlatformElectraDecodersWindows::RegisterWithCodecFactory(IElectraCodecRegi
 #endif
 	}
 }
+
+
+bool FElectraDecodersPlatformResourcesWindows::GetD3DDeviceAndVersion(void** OutD3DDevice, int32* OutD3DVersion)
+{
+	if (!GDynamicRHI || !OutD3DDevice || !OutD3DVersion)
+	{
+		return false;
+	}
+	auto RHIType = RHIGetInterfaceType();
+	if (RHIType != ERHIInterfaceType::D3D11 && RHIType != ERHIInterfaceType::D3D12)
+	{
+		return false;
+	}
+	*OutD3DDevice = GDynamicRHI->RHIGetNativeDevice();
+	*OutD3DVersion = RHIType == ERHIInterfaceType::D3D11 ? 11000 : 12000;
+	return true;
+}
+
+
+
+class FElectraDecodersPlatformResourcesWindows::FAsyncConsecutiveTaskSync
+{
+public:
+	FAsyncConsecutiveTaskSync() = default;
+	~FAsyncConsecutiveTaskSync() = default;
+	FGraphEventRef GraphEvent;
+};
+
+TSharedPtr<FElectraDecodersPlatformResourcesWindows::FAsyncConsecutiveTaskSync, ESPMode::ThreadSafe> FElectraDecodersPlatformResourcesWindows::CreateAsyncConsecutiveTaskSync()
+{
+	return MakeShared<FElectraDecodersPlatformResourcesWindows::FAsyncConsecutiveTaskSync, ESPMode::ThreadSafe>();
+}
+
+bool FElectraDecodersPlatformResourcesWindows::RunCodeAsync(TFunction<void()>&& InCodeToRun, TSharedPtr<FElectraDecodersPlatformResourcesWindows::FAsyncConsecutiveTaskSync, ESPMode::ThreadSafe> InTaskSync)
+{
+	FScopeLock Lock(&PlatformElectraDecodersWindows::AsyncJobAccessCS);
+
+	// Execute code async
+	// (We assume this to be code copying buffer data around. Hence we allow only ONE at any time to not clog up
+	//  the buses even more and delay the copy process further)
+	FGraphEventArray Events;
+	auto& Event = InTaskSync ? InTaskSync->GraphEvent : PlatformElectraDecodersWindows::RunCodeAsyncEvent;
+	if (Event.IsValid())
+	{
+		Events.Add(Event);
+	}
+	Event = FFunctionGraphTask::CreateAndDispatchWhenReady(MoveTemp(InCodeToRun), GET_STATID(STAT_ElectraDecoderAsyncJob), &Events);
+	return true;
+}
diff --git a/Engine/Plugins/Media/ElectraCodecs/Source/ElectraDecoders/Private/Windows/aac/AAC_AudioDecoder_Windows.cpp b/Engine/Plugins/Media/ElectraCodecs/Source/ElectraDecoders/Private/Windows/aac/AAC_AudioDecoder_Windows.cpp
--- a/Engine/Plugins/Media/ElectraCodecs/Source/ElectraDecoders/Private/Windows/aac/AAC_AudioDecoder_Windows.cpp
+++ b/Engine/Plugins/Media/ElectraCodecs/Source/ElectraDecoders/Private/Windows/aac/AAC_AudioDecoder_Windows.cpp
@@ -66,9 +66,9 @@ public:
 		IElectraAudioDecoderAAC_DX::GetConfigurationOptions(OutOptions);
 	}
 
-	TSharedPtr<IElectraDecoder, ESPMode::ThreadSafe> CreateDecoderForFormat(const FString& InCodecFormat, const TMap<FString, FVariant>& InOptions, TSharedPtr<IElectraDecoderResourceDelegate, ESPMode::ThreadSafe> InResourceDelegate) override
+	TSharedPtr<IElectraDecoder, ESPMode::ThreadSafe> CreateDecoderForFormat(const FString& InCodecFormat, const TMap<FString, FVariant>& InOptions) override
 	{
-		return IElectraAudioDecoderAAC_DX::Create(InOptions, InResourceDelegate);
+		return IElectraAudioDecoderAAC_DX::Create(InOptions);
 	}
 
 };
diff --git a/Engine/Plugins/Media/ElectraCodecs/Source/ElectraDecoders/Private/Windows/h264/H264_VideoDecoder_Windows.cpp b/Engine/Plugins/Media/ElectraCodecs/Source/ElectraDecoders/Private/Windows/h264/H264_VideoDecoder_Windows.cpp
--- a/Engine/Plugins/Media/ElectraCodecs/Source/ElectraDecoders/Private/Windows/h264/H264_VideoDecoder_Windows.cpp
+++ b/Engine/Plugins/Media/ElectraCodecs/Source/ElectraDecoders/Private/Windows/h264/H264_VideoDecoder_Windows.cpp
@@ -8,6 +8,8 @@
 #include "DX/VideoDecoderH264_DX.h"
 #include "DX/DecoderErrors_DX.h"
 
+#include "ElectraDecodersPlatformResources.h"
+
 #include "ElectraDecodersUtils.h"
 #include "IElectraDecoderFeaturesAndOptions.h"
 
@@ -115,9 +117,9 @@ public:
 		IElectraVideoDecoderH264_DX::PlatformGetConfigurationOptions(OutOptions);
 	}
 
-	TSharedPtr<IElectraDecoder, ESPMode::ThreadSafe> CreateDecoderForFormat(const FString& InCodecFormat, const TMap<FString, FVariant>& InOptions, TSharedPtr<IElectraDecoderResourceDelegate, ESPMode::ThreadSafe> InResourceDelegate) override
+	TSharedPtr<IElectraDecoder, ESPMode::ThreadSafe> CreateDecoderForFormat(const FString& InCodecFormat, const TMap<FString, FVariant>& InOptions) override
 	{
-		return IElectraVideoDecoderH264_DX::Create(InOptions, InResourceDelegate);
+		return IElectraVideoDecoderH264_DX::Create(InOptions);
 	}
 
 };
@@ -209,15 +211,15 @@ void IElectraVideoDecoderH264_DX::PlatformGetConfigurationOptions(TMap<FString,
 }
 
 
-IElectraDecoder::FError IElectraVideoDecoderH264_DX::PlatformCreateMFDecoderTransform(IPlatformHandle** OutPlatformHandle, TSharedPtr<IElectraDecoderResourceDelegate, ESPMode::ThreadSafe> InResourceDelegate, const TMap<FString, FVariant>& InOptions)
+IElectraDecoder::FError IElectraVideoDecoderH264_DX::PlatformCreateMFDecoderTransform(IPlatformHandle** OutPlatformHandle, const TMap<FString, FVariant>& InOptions)
 {
 	IElectraDecoder::FError Error;
-	check(OutPlatformHandle && InResourceDelegate.IsValid());
+	check(OutPlatformHandle);
 	if (OutPlatformHandle)
 	{
 		*OutPlatformHandle = nullptr;
 	}
-	if (OutPlatformHandle && InResourceDelegate.IsValid())
+	if (OutPlatformHandle)
 	{
 		// Must be Windows 8 or better
 		if (FPlatformMisc::VerifyWindowsVersion(6, 2))
@@ -280,7 +282,7 @@ IElectraDecoder::FError IElectraVideoDecoderH264_DX::PlatformCreateMFDecoderTran
 				TRefCountPtr<IDXGIAdapter> DXGIAdapter;
 				void* ApplicationD3DDevice = nullptr;
 				int32 ApplicationD3DDeviceVersion = 0;
-				if (InResourceDelegate->GetD3DDevice(&ApplicationD3DDevice, &ApplicationD3DDeviceVersion))
+				if (FElectraDecodersPlatformResources::GetD3DDeviceAndVersion(&ApplicationD3DDevice, &ApplicationD3DDeviceVersion))
 				{
 					if (ApplicationD3DDeviceVersion >= 12000)
 					{
@@ -377,10 +379,6 @@ IElectraDecoder::FError IElectraVideoDecoderH264_DX::PlatformCreateMFDecoderTran
 						DxMultithread->SetMultithreadProtected(1);
 					}
 					bAssociatedDeviceManager = true;
-
-
-
-
 				}
 #endif
 
