Changed to a warning when media foundation dlls are not present. They are not included on all server versions of Windows.
#rnx

diff --git a/Engine/Plugins/Media/ElectraCodecs/Source/ElectraDecoders/Private/Windows/WindowsPlatformElectraDecoders.cpp b/Engine/Plugins/Media/ElectraCodecs/Source/ElectraDecoders/Private/Windows/WindowsPlatformElectraDecoders.cpp
--- a/Engine/Plugins/Media/ElectraCodecs/Source/ElectraDecoders/Private/Windows/WindowsPlatformElectraDecoders.cpp
+++ b/Engine/Plugins/Media/ElectraCodecs/Source/ElectraDecoders/Private/Windows/WindowsPlatformElectraDecoders.cpp
@@ -56,7 +56,7 @@ void FPlatformElectraDecodersWindows::Startup()
 		&& FPlatformProcess::GetDllHandle(TEXT("msmpeg2vdec.dll"))
 		&& FPlatformProcess::GetDllHandle(TEXT("MSAudDecMFT.dll"))))
 	{
-		UE_LOG(LogElectraDecoders, Error, TEXT("Could not load Media Foundation"));
+		UE_LOG(LogElectraDecoders, Warning, TEXT("Could not load Media Foundation. One or more of these dlls not found - mf.dll, mfplat.dll, msmpeg2vdec.dll, MSAudDecMFT.dll. Usually this indicators running on a Windows Server OS."));
 		return;
 	}
 	HRESULT Result = MFStartup(MF_VERSION);
