Reworked ElectraSamples, ElectraUtils and the decoder output of Electra Player

#rnx

diff --git a/Engine/Plugins/Media/AppleProResDecoderElectra/Source/AppleProResDecoderElectra/AppleProResDecoderElectra.Build.cs b/Engine/Plugins/Media/AppleProResDecoderElectra/Source/AppleProResDecoderElectra/AppleProResDecoderElectra.Build.cs
--- a/Engine/Plugins/Media/AppleProResDecoderElectra/Source/AppleProResDecoderElectra/AppleProResDecoderElectra.Build.cs
+++ b/Engine/Plugins/Media/AppleProResDecoderElectra/Source/AppleProResDecoderElectra/AppleProResDecoderElectra.Build.cs
@@ -12,6 +12,8 @@ namespace UnrealBuildTool.Rules
 				new string[] {
 					"Core",
                     "Projects",
+					"ElectraBase",
+                    "ElectraSamples",
 					"ElectraCodecFactory",
 					"ElectraDecoders"
                 });
diff --git a/Engine/Plugins/Media/AppleProResDecoderElectra/Source/AppleProResDecoderElectra/Private/ProResDecoder/Mac/ElectraMediaProResDecoder.cpp b/Engine/Plugins/Media/AppleProResDecoderElectra/Source/AppleProResDecoderElectra/Private/ProResDecoder/Mac/ElectraMediaProResDecoder.cpp
--- a/Engine/Plugins/Media/AppleProResDecoderElectra/Source/AppleProResDecoderElectra/Private/ProResDecoder/Mac/ElectraMediaProResDecoder.cpp
+++ b/Engine/Plugins/Media/AppleProResDecoderElectra/Source/AppleProResDecoderElectra/Private/ProResDecoder/Mac/ElectraMediaProResDecoder.cpp
@@ -13,7 +13,7 @@
 #include "IElectraDecoder.h"
 #include "IElectraDecoderOutputVideo.h"
 #include "ElectraDecodersUtils.h"
-#include "Utils/MPEG/ElectraUtilsMP4.h"	// to read the ProRes header, which is given in the form of an MP4 Atom
+#include "Utilities/UtilitiesMP4.h"	// to read the ProRes header, which is given in the form of an MP4 Atom
 #include "Utils/ElectraBitstreamProcessorGenericVideo.h"
 #include "AppleProResDecoderElectraModule.h"
 #include "Apple/AppleElectraDecoderPlatformOutputFormatTypes.h"
@@ -98,7 +98,7 @@ public:
 	{
 		return 1;
 	}
-	TSharedPtr<TArray<uint8>, ESPMode::ThreadSafe> GetBufferDataByIndex(int32 InBufferIndex) const override
+	TSharedPtr<TArray64<uint8>, ESPMode::ThreadSafe> GetBufferDataByIndex(int32 InBufferIndex) const override
 	{
 		return nullptr;
 	}
@@ -110,11 +110,11 @@ public:
 	{
 		return false;
 	}
-	EElectraDecoderPlatformPixelFormat GetBufferFormatByIndex(int32 InBufferIndex) const override
+	EPixelFormat GetBufferFormatByIndex(int32 InBufferIndex) const override
 	{
 		return PixelFormat;
 	}
-	EElectraDecoderPlatformPixelEncoding GetBufferEncodingByIndex(int32 InBufferIndex) const override
+	EElectraTextureSamplePixelEncoding GetBufferEncodingByIndex(int32 InBufferIndex) const override
 	{
 		return PixelEncoding;
 	}
@@ -145,8 +145,8 @@ public:
 	int32 AspectH = 1;
 	int32 FrameRateN = 0;
 	int32 FrameRateD = 0;
-	EElectraDecoderPlatformPixelFormat PixelFormat = EElectraDecoderPlatformPixelFormat::INVALID;
-	EElectraDecoderPlatformPixelEncoding PixelEncoding = EElectraDecoderPlatformPixelEncoding::Native;
+	EPixelFormat PixelFormat = EPixelFormat::PF_Unknown;
+	EElectraTextureSamplePixelEncoding PixelEncoding = EElectraTextureSamplePixelEncoding::Native;
 	TMap<FString, FVariant> ExtraValues;
 
 	uint32 Codec4CC = 0;
@@ -714,10 +714,10 @@ bool FVideoDecoderProResElectra::ParseHeader(const void* InData, int32 InNumData
 		return false;
 	}
 	CurrentHeader.Reset();
-	ElectraDecodersUtil::FMP4AtomReader Reader(InData, InNumData);
+	Electra::UtilitiesMP4::FMP4AtomReader Reader(MakeConstArrayView<const uint8>(reinterpret_cast<const uint8*>(InData), InNumData));
 	Reader.Read(CurrentHeader.FrameSizeInBytes);
 	Reader.Read(CurrentHeader.FrameType);
-	if (CurrentHeader.FrameType != ElectraDecodersUtil::MakeMP4Atom('i','c','p','f'))
+	if (CurrentHeader.FrameType != Electra::UtilitiesMP4::MakeBoxAtom('i','c','p','f'))
 	{
 		return false;
 	}
@@ -1157,33 +1157,33 @@ FVideoDecoderProResElectra::EConvertResult FVideoDecoderProResElectra::ConvertDe
 
 	NewOutput->Codec4CC = Codec4CC;
 
-	switch (pixelFormat)
+	switch(pixelFormat)
 	{
-		case	kCVPixelFormatType_4444AYpCbCr16:
+		case kCVPixelFormatType_4444AYpCbCr16:
 		{
-			NewOutput->PixelFormat = EElectraDecoderPlatformPixelFormat::R16G16B16A16;
-			NewOutput->PixelEncoding = EElectraDecoderPlatformPixelEncoding::YCbCr_Alpha;
+			NewOutput->PixelFormat = EPixelFormat::PF_A16B16G16R16;
+			NewOutput->PixelEncoding = EElectraTextureSamplePixelEncoding::YCbCr_Alpha;
 			break;
 		}
-		case	kCVPixelFormatType_420YpCbCr8BiPlanarVideoRange:
-		case	kCVPixelFormatType_420YpCbCr8BiPlanarFullRange:
+		case kCVPixelFormatType_420YpCbCr8BiPlanarVideoRange:
+		case kCVPixelFormatType_420YpCbCr8BiPlanarFullRange:
 		{
-			NewOutput->PixelFormat = EElectraDecoderPlatformPixelFormat::NV12;
-			NewOutput->PixelEncoding = EElectraDecoderPlatformPixelEncoding::Native;
+			NewOutput->PixelFormat = EPixelFormat::PF_NV12;
+			NewOutput->PixelEncoding = EElectraTextureSamplePixelEncoding::Native;
 			break;
 		}
-		case	kCVPixelFormatType_420YpCbCr10BiPlanarVideoRange:
-		case	kCVPixelFormatType_420YpCbCr10BiPlanarFullRange:
+		case kCVPixelFormatType_420YpCbCr10BiPlanarVideoRange:
+		case kCVPixelFormatType_420YpCbCr10BiPlanarFullRange:
 		{
-			NewOutput->PixelFormat = EElectraDecoderPlatformPixelFormat::P010;
-			NewOutput->PixelEncoding = EElectraDecoderPlatformPixelEncoding::Native;
+			NewOutput->PixelFormat = EPixelFormat::PF_P010;
+			NewOutput->PixelEncoding = EElectraTextureSamplePixelEncoding::Native;
 			break;
 		}
 		default:
 		{
 			check(!"Unexpected output format!");
-			NewOutput->PixelFormat = EElectraDecoderPlatformPixelFormat::INVALID;
-			NewOutput->PixelEncoding = EElectraDecoderPlatformPixelEncoding::Native;
+			NewOutput->PixelFormat = EPixelFormat::PF_Unknown;
+			NewOutput->PixelEncoding = EElectraTextureSamplePixelEncoding::Native;
 		}
 	}
 
diff --git a/Engine/Plugins/Media/AppleProResDecoderElectra/Source/AppleProResDecoderElectra/Private/ProResDecoder/Windows/ElectraMediaProResDecoder.cpp b/Engine/Plugins/Media/AppleProResDecoderElectra/Source/AppleProResDecoderElectra/Private/ProResDecoder/Windows/ElectraMediaProResDecoder.cpp
--- a/Engine/Plugins/Media/AppleProResDecoderElectra/Source/AppleProResDecoderElectra/Private/ProResDecoder/Windows/ElectraMediaProResDecoder.cpp
+++ b/Engine/Plugins/Media/AppleProResDecoderElectra/Source/AppleProResDecoderElectra/Private/ProResDecoder/Windows/ElectraMediaProResDecoder.cpp
@@ -13,11 +13,11 @@
 #include "IElectraDecoder.h"
 #include "IElectraDecoderOutputVideo.h"
 #include "ElectraDecodersUtils.h"
-#include "Utils/MPEG/ElectraUtilsMP4.h"	// to read the ProRes header, which is given in the form of an MP4 Atom
+#include "Utilities/UtilitiesMP4.h"	// to read the ProRes header, which is given in the form of an MP4 Atom
 #include "Utils/ElectraBitstreamProcessorGenericVideo.h"
 
 #include "ElectraDecodersPlatformResources.h"
-#include COMPILED_PLATFORM_HEADER(ElectraDecoderGPUBufferHelpers.h)
+#include COMPILED_PLATFORM_HEADER(ElectraTextureSampleGPUBufferHelper.h)
 
 #include "ProResDecoder.h"
 
@@ -119,7 +119,7 @@ public:
 	{
 		return 1;
 	}
-	TSharedPtr<TArray<uint8>, ESPMode::ThreadSafe> GetBufferDataByIndex(int32 InBufferIndex) const override
+	TSharedPtr<TArray64<uint8>, ESPMode::ThreadSafe> GetBufferDataByIndex(int32 InBufferIndex) const override
 	{
 		return InBufferIndex == 0 ? Buffer : nullptr;
 	}
@@ -144,11 +144,11 @@ public:
 #endif
 		return false;
 	}
-	EElectraDecoderPlatformPixelFormat GetBufferFormatByIndex(int32 InBufferIndex) const override
+	EPixelFormat GetBufferFormatByIndex(int32 InBufferIndex) const override
 	{
 		return BufferFormat;
 	}
-	EElectraDecoderPlatformPixelEncoding GetBufferEncodingByIndex(int32 InBufferIndex) const override
+	EElectraTextureSamplePixelEncoding GetBufferEncodingByIndex(int32 InBufferIndex) const override
 	{
 		return BufferEncoding;
 	}
@@ -176,9 +176,9 @@ public:
 	TMap<FString, FVariant> ExtraValues;
 
 	uint32 Codec4CC = 0;
-	TSharedPtr<TArray<uint8>, ESPMode::ThreadSafe> Buffer;
-	EElectraDecoderPlatformPixelFormat BufferFormat = EElectraDecoderPlatformPixelFormat::INVALID;
-	EElectraDecoderPlatformPixelEncoding BufferEncoding = EElectraDecoderPlatformPixelEncoding::Native;
+	TSharedPtr<TArray64<uint8>, ESPMode::ThreadSafe> Buffer;
+	EPixelFormat BufferFormat = EPixelFormat::PF_Unknown;
+	EElectraTextureSamplePixelEncoding BufferEncoding = EElectraTextureSamplePixelEncoding::Native;
 #if ELECTRA_MEDIAGPUBUFFER_DX12
 	FElectraMediaDecoderOutputBufferPool_DX12::FOutputData GPUBuffer;
 #endif
@@ -574,10 +574,10 @@ bool FVideoDecoderProResElectra::ParseHeader(const void* InData, int32 InNumData
 		return false;
 	}
 	CurrentHeader.Reset();
-	ElectraDecodersUtil::FMP4AtomReader Reader(InData, InNumData);
+	Electra::UtilitiesMP4::FMP4AtomReader Reader(MakeConstArrayView<const uint8>(reinterpret_cast<const uint8*>(InData), InNumData));
 	Reader.Read(CurrentHeader.FrameSizeInBytes);
 	Reader.Read(CurrentHeader.FrameType);
-	if (CurrentHeader.FrameType != ElectraDecodersUtil::MakeMP4Atom('i','c','p','f'))
+	if (CurrentHeader.FrameType != Electra::UtilitiesMP4::MakeBoxAtom('i','c','p','f'))
 	{
 		return false;
 	}
@@ -736,48 +736,48 @@ IElectraDecoder::EDecoderError FVideoDecoderProResElectra::DecodeAccessUnit(cons
 		{
 			case kPRFormat_2vuy:
 				NewOutput->NumBits = 8;
-				NewOutput->BufferFormat = EElectraDecoderPlatformPixelFormat::B8G8R8A8;
-				NewOutput->BufferEncoding = EElectraDecoderPlatformPixelEncoding::CbY0CrY1;
+				NewOutput->BufferFormat = EPixelFormat::PF_B8G8R8A8;
+				NewOutput->BufferEncoding = EElectraTextureSamplePixelEncoding::CbY0CrY1;
 				NewOutput->Width /= 2;
 				NewOutput->Pitch *= 2;
 				break;
 			case kPRFormat_y416:
 				NewOutput->NumBits = 16;
-				NewOutput->BufferFormat = EElectraDecoderPlatformPixelFormat::A16B16G16R16;
-				NewOutput->BufferEncoding = EElectraDecoderPlatformPixelEncoding::YCbCr_Alpha;
+				NewOutput->BufferFormat = EPixelFormat::PF_A16B16G16R16;
+				NewOutput->BufferEncoding = EElectraTextureSamplePixelEncoding::YCbCr_Alpha;
 				NewOutput->Pitch *= 8;
 				break;
 			case kPRFormat_r4fl:
 				NewOutput->NumBits = 32;
-				NewOutput->BufferFormat = EElectraDecoderPlatformPixelFormat::A32B32G32R32F;
-				NewOutput->BufferEncoding = EElectraDecoderPlatformPixelEncoding::YCbCr_Alpha;
+				NewOutput->BufferFormat = EPixelFormat::PF_A32B32G32R32F;
+				NewOutput->BufferEncoding = EElectraTextureSamplePixelEncoding::YCbCr_Alpha;
 				NewOutput->Pitch *= 16;
 				break;
 			case kPRFormat_v210:
 				NewOutput->NumBits = 10;
-				NewOutput->BufferFormat = EElectraDecoderPlatformPixelFormat::A2B10G10R10;
-				NewOutput->BufferEncoding = EElectraDecoderPlatformPixelEncoding::CbY0CrY1;
+				NewOutput->BufferFormat = EPixelFormat::PF_A2B10G10R10;
+				NewOutput->BufferEncoding = EElectraTextureSamplePixelEncoding::CbY0CrY1;
 				NewOutput->Width = 4 * ((NewOutput->Width + 5) / 6); // each 4 pixel contain 6 horizontally adjacent YCbCr pixels (incl. 4x 2-bit padding)
 				NewOutput->Pitch = NewOutput->Width * 4;
 				break;
 			case kPRFormat_v216:
 				NewOutput->NumBits = 16;
-				NewOutput->BufferFormat = EElectraDecoderPlatformPixelFormat::A16B16G16R16;
-				NewOutput->BufferEncoding = EElectraDecoderPlatformPixelEncoding::CbY0CrY1;
+				NewOutput->BufferFormat = EPixelFormat::PF_A16B16G16R16;
+				NewOutput->BufferEncoding = EElectraTextureSamplePixelEncoding::CbY0CrY1;
 				NewOutput->Width /= 2;
 				NewOutput->Pitch *= 4;
 				break;
 			case kPRFormat_b64a:
 				NewOutput->NumBits = 16;
-				NewOutput->BufferFormat = EElectraDecoderPlatformPixelFormat::A16B16G16R16;
-				NewOutput->BufferEncoding = EElectraDecoderPlatformPixelEncoding::ARGB_BigEndian;
+				NewOutput->BufferFormat = EPixelFormat::PF_A16B16G16R16;
+				NewOutput->BufferEncoding = EElectraTextureSamplePixelEncoding::ARGB_BigEndian;
 				NewOutput->Pitch *= 8;
 				break;
 			case kPRFormat_R10k:
 			case kPRFormat_r210:
 			default:
 				NewOutput->NumBits = 0;
-				NewOutput->BufferFormat = EElectraDecoderPlatformPixelFormat::INVALID;
+				NewOutput->BufferFormat = EPixelFormat::PF_Unknown;
 				break;
 		}
 
@@ -823,7 +823,7 @@ IElectraDecoder::EDecoderError FVideoDecoderProResElectra::DecodeAccessUnit(cons
 		else
 #endif
 		{
-			NewOutput->Buffer = MakeShared<TArray<uint8>, ESPMode::ThreadSafe>();
+			NewOutput->Buffer = MakeShared<TArray64<uint8>, ESPMode::ThreadSafe>();
 			NewOutput->Buffer->AddUninitialized(BufferAllocationSize);
 
 			pb.baseAddr = NewOutput->Buffer->GetData();
diff --git a/Engine/Plugins/Media/AvidDNxHDDecoderElectra/Source/AvidDNxHDDecoderElectra.Build.cs b/Engine/Plugins/Media/AvidDNxHDDecoderElectra/Source/AvidDNxHDDecoderElectra.Build.cs
--- a/Engine/Plugins/Media/AvidDNxHDDecoderElectra/Source/AvidDNxHDDecoderElectra.Build.cs
+++ b/Engine/Plugins/Media/AvidDNxHDDecoderElectra/Source/AvidDNxHDDecoderElectra.Build.cs
@@ -12,6 +12,8 @@ namespace UnrealBuildTool.Rules
 				new string[] {
 					"Core",
                     "Projects",
+					"ElectraBase",
+                    "ElectraSamples",
 					"ElectraCodecFactory",
 					"ElectraDecoders"
                 });
diff --git a/Engine/Plugins/Media/AvidDNxHDDecoderElectra/Source/Private/Windows/ElectraMediaAvidDNxHDDecoder.cpp b/Engine/Plugins/Media/AvidDNxHDDecoderElectra/Source/Private/Windows/ElectraMediaAvidDNxHDDecoder.cpp
--- a/Engine/Plugins/Media/AvidDNxHDDecoderElectra/Source/Private/Windows/ElectraMediaAvidDNxHDDecoder.cpp
+++ b/Engine/Plugins/Media/AvidDNxHDDecoderElectra/Source/Private/Windows/ElectraMediaAvidDNxHDDecoder.cpp
@@ -20,7 +20,7 @@
 #include "Utils/ElectraBitstreamProcessorDefault.h"
 
 #include "ElectraDecodersPlatformResources.h"
-#include COMPILED_PLATFORM_HEADER(ElectraDecoderGPUBufferHelpers.h)
+#include COMPILED_PLATFORM_HEADER(ElectraTextureSampleGPUBufferHelper.h)
 
 #include "AvidDNxCodec.h"
 //#include "dnx_uncompressed_sdk.h"
@@ -123,7 +123,7 @@ public:
 	{
 		return 1;
 	}
-	TSharedPtr<TArray<uint8>, ESPMode::ThreadSafe> GetBufferDataByIndex(int32 InBufferIndex) const override
+	TSharedPtr<TArray64<uint8>, ESPMode::ThreadSafe> GetBufferDataByIndex(int32 InBufferIndex) const override
 	{
 		if (InBufferIndex == 0)
 		{
@@ -152,21 +152,21 @@ public:
 #endif
 	return false;
 	}
-	EElectraDecoderPlatformPixelFormat GetBufferFormatByIndex(int32 InBufferIndex) const override
+	EPixelFormat GetBufferFormatByIndex(int32 InBufferIndex) const override
 	{
 		if (InBufferIndex == 0)
 		{
 			return BufferFormat;
 		}
-		return EElectraDecoderPlatformPixelFormat::INVALID;
+		return EPixelFormat::PF_Unknown;
 	}
-	EElectraDecoderPlatformPixelEncoding GetBufferEncodingByIndex(int32 InBufferIndex) const override
+	EElectraTextureSamplePixelEncoding GetBufferEncodingByIndex(int32 InBufferIndex) const override
 	{
 		if (InBufferIndex == 0)
 		{
 			return BufferEncoding;
 		}
-		return EElectraDecoderPlatformPixelEncoding::Native;
+		return EElectraTextureSamplePixelEncoding::Native;
 	}
 	int32 GetBufferPitchByIndex(int32 InBufferIndex) const override
 	{
@@ -196,9 +196,9 @@ public:
 	TMap<FString, FVariant> ExtraValues;
 
 	uint32 Codec4CC = 0;
-	TSharedPtr<TArray<uint8>, ESPMode::ThreadSafe> Buffer;
-	EElectraDecoderPlatformPixelFormat BufferFormat = EElectraDecoderPlatformPixelFormat::INVALID;
-	EElectraDecoderPlatformPixelEncoding BufferEncoding = EElectraDecoderPlatformPixelEncoding::Native;
+	TSharedPtr<TArray64<uint8>, ESPMode::ThreadSafe> Buffer;
+	EPixelFormat BufferFormat = EPixelFormat::PF_Unknown;
+	EElectraTextureSamplePixelEncoding BufferEncoding = EElectraTextureSamplePixelEncoding::Native;
 #if ELECTRA_MEDIAGPUBUFFER_DX12
 	FElectraMediaDecoderOutputBufferPool_DX12::FOutputData GPUBuffer;
 	TSharedPtr<FElectraDecodersPlatformResources::FAsyncConsecutiveTaskSync, ESPMode::ThreadSafe> GPUBuffer_TaskSync;
@@ -571,46 +571,46 @@ IElectraDecoder::EDecoderError FVideoDecoderAvidDNxHDElectra::DecodeAccessUnit(c
 		NewOutput->PixelFormat = (int32) Decoder.CurrentUncompressedParams.compType;
 
 		bool bIsPlanar = false;
-		switch (Decoder.CurrentUncompressedParams.compType)
+		switch(Decoder.CurrentUncompressedParams.compType)
 		{
-			case	DNX_CT_UCHAR:
+			case DNX_CT_UCHAR:
 			{
-				switch (Decoder.CurrentUncompressedParams.compOrder)
+				switch(Decoder.CurrentUncompressedParams.compOrder)
 				{
-					case	DNX_CCO_YCbYCr_NoA:
-					case	DNX_CCO_CbYCrY_NoA:
+					case DNX_CCO_YCbYCr_NoA:
+					case DNX_CCO_CbYCrY_NoA:
 					{
-						NewOutput->BufferFormat = EElectraDecoderPlatformPixelFormat::B8G8R8A8;
-						NewOutput->BufferEncoding = (Decoder.CurrentUncompressedParams.compOrder  == DNX_CCO_CbYCrY_NoA) ?  EElectraDecoderPlatformPixelEncoding::CbY0CrY1 : EElectraDecoderPlatformPixelEncoding::Y0CbY1Cr;
+						NewOutput->BufferFormat = EPixelFormat::PF_B8G8R8A8;
+						NewOutput->BufferEncoding = (Decoder.CurrentUncompressedParams.compOrder == DNX_CCO_CbYCrY_NoA) ?  EElectraTextureSamplePixelEncoding::CbY0CrY1 : EElectraTextureSamplePixelEncoding::Y0CbY1Cr;
 						NewOutput->Width /= 2;
 						NewOutput->Pitch *= 2;
 						break;
 					}
-					case	DNX_CCO_ARGB_Interleaved:
+					case DNX_CCO_ARGB_Interleaved:
 					{
-						NewOutput->BufferFormat = EElectraDecoderPlatformPixelFormat::A8R8G8B8;
-						NewOutput->BufferEncoding = EElectraDecoderPlatformPixelEncoding::Native;
+						NewOutput->BufferFormat = EPixelFormat::PF_A8R8G8B8;
+						NewOutput->BufferEncoding = EElectraTextureSamplePixelEncoding::Native;
 						NewOutput->Pitch *= 4;
 						break;
 					}
-					case	DNX_CCO_RGBA_Interleaved:
+					case DNX_CCO_RGBA_Interleaved:
 					{
-						NewOutput->BufferFormat = EElectraDecoderPlatformPixelFormat::R8G8B8A8;
-						NewOutput->BufferEncoding = EElectraDecoderPlatformPixelEncoding::Native;
+						NewOutput->BufferFormat = EPixelFormat::PF_R8G8B8A8;
+						NewOutput->BufferEncoding = EElectraTextureSamplePixelEncoding::Native;
 						NewOutput->Pitch *= 4;
 						break;
 					}
-					case	DNX_CCO_CbYCrA_Interleaved:
+					case DNX_CCO_CbYCrA_Interleaved:
 					{
-						NewOutput->BufferFormat = EElectraDecoderPlatformPixelFormat::B8G8R8A8;
-						NewOutput->BufferEncoding = EElectraDecoderPlatformPixelEncoding::YCbCr_Alpha;
+						NewOutput->BufferFormat = EPixelFormat::PF_B8G8R8A8;
+						NewOutput->BufferEncoding = EElectraTextureSamplePixelEncoding::YCbCr_Alpha;
 						NewOutput->Pitch *= 4;
 						break;
 					}
-					case	DNX_CCO_YCbCr_Planar:
+					case DNX_CCO_YCbCr_Planar:
 					{
-						NewOutput->BufferFormat = EElectraDecoderPlatformPixelFormat::NV12;
-						NewOutput->BufferEncoding = EElectraDecoderPlatformPixelEncoding::Native;
+						NewOutput->BufferFormat = EPixelFormat::PF_NV12;
+						NewOutput->BufferEncoding = EElectraTextureSamplePixelEncoding::Native;
 						if (!bUseGPUBuffers)
 						{
 							NewOutput->Height = (NewOutput->Height * 3) / 2;
@@ -620,52 +620,52 @@ IElectraDecoder::EDecoderError FVideoDecoderAvidDNxHDElectra::DecodeAccessUnit(c
 					}
 					default:
 					{
-						NewOutput->BufferFormat = EElectraDecoderPlatformPixelFormat::INVALID;
-						NewOutput->BufferEncoding = EElectraDecoderPlatformPixelEncoding::Native;
+						NewOutput->BufferFormat = EPixelFormat::PF_Unknown;
+						NewOutput->BufferEncoding = EElectraTextureSamplePixelEncoding::Native;
 					}
 				}
 				break;
 			}
-			case	DNX_CT_SHORT:
-			case	DNX_CT_USHORT_10_6:
-			case	DNX_CT_USHORT_12_4:
+			case DNX_CT_SHORT:
+			case DNX_CT_USHORT_10_6:
+			case DNX_CT_USHORT_12_4:
 			{
-				switch (Decoder.CurrentUncompressedParams.compOrder)
+				switch(Decoder.CurrentUncompressedParams.compOrder)
 				{
-					case	DNX_CCO_YCbYCr_NoA:
-					case	DNX_CCO_CbYCrY_NoA:
+					case DNX_CCO_YCbYCr_NoA:
+					case DNX_CCO_CbYCrY_NoA:
 					{
-						NewOutput->BufferFormat = EElectraDecoderPlatformPixelFormat::A16B16G16R16;
-						NewOutput->BufferEncoding = (Decoder.CurrentUncompressedParams.compOrder == DNX_CCO_CbYCrY_NoA) ? EElectraDecoderPlatformPixelEncoding::CbY0CrY1 : EElectraDecoderPlatformPixelEncoding::Y0CbY1Cr;
+						NewOutput->BufferFormat = EPixelFormat::PF_A16B16G16R16;
+						NewOutput->BufferEncoding = (Decoder.CurrentUncompressedParams.compOrder == DNX_CCO_CbYCrY_NoA) ? EElectraTextureSamplePixelEncoding::CbY0CrY1 : EElectraTextureSamplePixelEncoding::Y0CbY1Cr;
 						NewOutput->Width /= 2;
 						NewOutput->Pitch *= 4;
 						break;
 					}
 					case	DNX_CCO_RGBA_Interleaved:
 					{
-						NewOutput->BufferFormat = EElectraDecoderPlatformPixelFormat::R16G16B16A16;
-						NewOutput->BufferEncoding = EElectraDecoderPlatformPixelEncoding::Native;
+						NewOutput->BufferFormat = EPixelFormat::PF_R16G16B16A16_UNORM;
+						NewOutput->BufferEncoding = EElectraTextureSamplePixelEncoding::Native;
 						NewOutput->Pitch *= 8;
 						break;
 					}
 					case	DNX_CCO_ABGR_Interleaved:
 					{
-						NewOutput->BufferFormat = EElectraDecoderPlatformPixelFormat::A16B16G16R16;
-						NewOutput->BufferEncoding = EElectraDecoderPlatformPixelEncoding::Native;
+						NewOutput->BufferFormat = EPixelFormat::PF_A16B16G16R16;
+						NewOutput->BufferEncoding = EElectraTextureSamplePixelEncoding::Native;
 						NewOutput->Pitch *= 8;
 						break;
 					}
 					case	DNX_CCO_CbYCrA_Interleaved:
 					{
-						NewOutput->BufferFormat = EElectraDecoderPlatformPixelFormat::A16B16G16R16;
-						NewOutput->BufferEncoding = EElectraDecoderPlatformPixelEncoding::YCbCr_Alpha;
+						NewOutput->BufferFormat = EPixelFormat::PF_A16B16G16R16;
+						NewOutput->BufferEncoding = EElectraTextureSamplePixelEncoding::YCbCr_Alpha;
 						NewOutput->Pitch *= 8;
 						break;
 					}
 					case	DNX_CCO_YCbCr_Planar:
 					{
-						NewOutput->BufferFormat = EElectraDecoderPlatformPixelFormat::P010;
-						NewOutput->BufferEncoding = EElectraDecoderPlatformPixelEncoding::Native;
+						NewOutput->BufferFormat = EPixelFormat::PF_P010;
+						NewOutput->BufferEncoding = EElectraTextureSamplePixelEncoding::Native;
 						if (!bUseGPUBuffers)
 						{
 							NewOutput->Height = (NewOutput->Height * 3) / 2;
@@ -675,22 +675,22 @@ IElectraDecoder::EDecoderError FVideoDecoderAvidDNxHDElectra::DecodeAccessUnit(c
 					}
 					default:
 					{
-						NewOutput->BufferFormat = EElectraDecoderPlatformPixelFormat::INVALID;
-						NewOutput->BufferEncoding = EElectraDecoderPlatformPixelEncoding::Native;
+						NewOutput->BufferFormat = EPixelFormat::PF_Unknown;
+						NewOutput->BufferEncoding = EElectraTextureSamplePixelEncoding::Native;
 					}
 				}
 				break;
 			}
-			case	DNX_CT_V210:
+			case DNX_CT_V210:
 			{
-				NewOutput->BufferFormat = EElectraDecoderPlatformPixelFormat::A2B10G10R10;
-				NewOutput->BufferEncoding = EElectraDecoderPlatformPixelEncoding::CbY0CrY1;
+				NewOutput->BufferFormat = EPixelFormat::PF_A2B10G10R10;
+				NewOutput->BufferEncoding = EElectraTextureSamplePixelEncoding::CbY0CrY1;
 				break;
 			}
 			default:
 			{
-				NewOutput->BufferFormat = EElectraDecoderPlatformPixelFormat::INVALID;
-				NewOutput->BufferEncoding = EElectraDecoderPlatformPixelEncoding::Native;
+				NewOutput->BufferFormat = EPixelFormat::PF_Unknown;
+				NewOutput->BufferEncoding = EElectraTextureSamplePixelEncoding::Native;
 			}
 		}
 
@@ -755,7 +755,7 @@ IElectraDecoder::EDecoderError FVideoDecoderAvidDNxHDElectra::DecodeAccessUnit(c
 		else
 #endif
 		{
-			NewOutput->Buffer = MakeShared<TArray<uint8>, ESPMode::ThreadSafe>();
+			NewOutput->Buffer = MakeShared<TArray64<uint8>, ESPMode::ThreadSafe>();
 			NewOutput->Buffer->AddUninitialized(BPP * NewOutput->Width * NewOutput->Height);
 
 			OutBufferData = NewOutput->Buffer->GetData();
diff --git a/Engine/Plugins/Media/D3D12VideoDecodersElectra/Source/Private/VideoDecoder_D3D12_Common.h b/Engine/Plugins/Media/D3D12VideoDecodersElectra/Source/Private/VideoDecoder_D3D12_Common.h
--- a/Engine/Plugins/Media/D3D12VideoDecodersElectra/Source/Private/VideoDecoder_D3D12_Common.h
+++ b/Engine/Plugins/Media/D3D12VideoDecodersElectra/Source/Private/VideoDecoder_D3D12_Common.h
@@ -255,7 +255,7 @@ public:
 	{
 		return 1;
 	}
-	TSharedPtr<TArray<uint8>, ESPMode::ThreadSafe> GetBufferDataByIndex(int32 InBufferIndex) const override
+	TSharedPtr<TArray64<uint8>, ESPMode::ThreadSafe> GetBufferDataByIndex(int32 InBufferIndex) const override
 	{
 		check(InBufferIndex == 0);
 		// No CPU data here.
@@ -270,11 +270,11 @@ public:
 		}
 		return nullptr;
 	}
-	EElectraDecoderPlatformPixelFormat GetBufferFormatByIndex(int32 InBufferIndex) const override
+	EPixelFormat GetBufferFormatByIndex(int32 InBufferIndex) const override
 	{
 		return BufferFormat;
 	}
-	EElectraDecoderPlatformPixelEncoding GetBufferEncodingByIndex(int32 InBufferIndex) const override
+	EElectraTextureSamplePixelEncoding GetBufferEncodingByIndex(int32 InBufferIndex) const override
 	{
 		return BufferEncoding;
 	}
@@ -322,8 +322,8 @@ public:
 	bool bDoNotOutput = false;
 
 	uint32 Codec4CC = 0;
-	EElectraDecoderPlatformPixelFormat BufferFormat = EElectraDecoderPlatformPixelFormat::INVALID;
-	EElectraDecoderPlatformPixelEncoding BufferEncoding = EElectraDecoderPlatformPixelEncoding::Native;
+	EPixelFormat BufferFormat = EPixelFormat::PF_Unknown;
+	EElectraTextureSamplePixelEncoding BufferEncoding = EElectraTextureSamplePixelEncoding::Native;
 
 
 	TSharedPtr<FDecodedPictureBuffer, ESPMode::ThreadSafe> OwningDPB;
diff --git a/Engine/Plugins/Media/D3D12VideoDecodersElectra/Source/Private/VideoDecoder_D3D12_H264.cpp b/Engine/Plugins/Media/D3D12VideoDecodersElectra/Source/Private/VideoDecoder_D3D12_H264.cpp
--- a/Engine/Plugins/Media/D3D12VideoDecodersElectra/Source/Private/VideoDecoder_D3D12_H264.cpp
+++ b/Engine/Plugins/Media/D3D12VideoDecodersElectra/Source/Private/VideoDecoder_D3D12_H264.cpp
@@ -3,6 +3,7 @@
 #include "VideoDecoder_D3D12_H264.h"
 #include "D3D12VideoDecodersElectraModule.h"
 #include "Utils/MPEG/ElectraBitstreamProcessor_H264.h"
+#include "PixelFormat.h"
 
 
 namespace ElectraVideoDecodersD3D12Video
@@ -154,7 +155,7 @@ IElectraDecoder::EDecoderError FD3D12VideoDecoder_H264::DecodeAccessUnit(const F
 			FBitstreamParamsH264::FSliceDecodeInfo& SliceInfo = SliceInfos.Emplace_GetRef();
 			SliceInfo.NalUnitType = (uint8)nut;
 			SliceInfo.NalRefIdc = (uint8)refIdc;
-			ElectraDecodersUtil::MPEG::H264::FBitstreamReader br;
+			ElectraDecodersUtil::MPEG::H264::FBitstreamReaderH264 br;
 			TUniquePtr<ElectraDecodersUtil::MPEG::H264::FRBSP> SliceRBSP;
 			if (!ElectraDecodersUtil::MPEG::H264::ParseSliceHeader(SliceRBSP, br, SliceInfo.Header, BitstreamParamsH264.SPSs, BitstreamParamsH264.PPSs, Start+4, NaluLength))
 			{
@@ -638,8 +639,8 @@ IElectraDecoder::EDecoderError FD3D12VideoDecoder_H264::DecodeSlicesH264(const F
 	InDec->ImageHeight = InDec->Height - InDec->Crop.Top - InDec->Crop.Bottom;
 	//InDec->Pitch = InDec->ImageWidth;
 	InDec->NumBits = 8;
-	InDec->BufferFormat = EElectraDecoderPlatformPixelFormat::NV12;
-	InDec->BufferEncoding = EElectraDecoderPlatformPixelEncoding::Native;
+	InDec->BufferFormat = EPixelFormat::PF_NV12;
+	InDec->BufferEncoding = EElectraTextureSamplePixelEncoding::Native;
 	InSequenceParameterSet.GetAspect(InDec->AspectW, InDec->AspectH);
 	auto FrameRate = InSequenceParameterSet.GetTiming();
 	InDec->FrameRateN = FrameRate.Denom ? FrameRate.Num : 30;
@@ -649,8 +650,8 @@ IElectraDecoder::EDecoderError FD3D12VideoDecoder_H264::DecodeSlicesH264(const F
 	InDec->ExtraValues.Emplace(TEXT("dxversion"), FVariant((int64) 12000));
 	InDec->ExtraValues.Emplace(TEXT("sw"), FVariant(false));
 	InDec->ExtraValues.Emplace(TEXT("codec"), FVariant(TEXT("avc")));
-	InDec->ExtraValues.Emplace(TEXT("pixfmt"), FVariant((int64)EElectraDecoderPlatformPixelFormat::NV12));
-	InDec->ExtraValues.Emplace(TEXT("pixenc"), FVariant((int64)EElectraDecoderPlatformPixelEncoding::Native));
+	InDec->ExtraValues.Emplace(TEXT("pixfmt"), FVariant((int64)EPixelFormat::PF_NV12));
+	InDec->ExtraValues.Emplace(TEXT("pixenc"), FVariant((int64)EElectraTextureSamplePixelEncoding::Native));
 	FramesInDecoder.Emplace(MoveTemp(InDec));
 
 	// Update the simulation DPB with the new decoded frame.
diff --git a/Engine/Plugins/Media/D3D12VideoDecodersElectra/Source/Private/VideoDecoder_D3D12_H265.cpp b/Engine/Plugins/Media/D3D12VideoDecodersElectra/Source/Private/VideoDecoder_D3D12_H265.cpp
--- a/Engine/Plugins/Media/D3D12VideoDecodersElectra/Source/Private/VideoDecoder_D3D12_H265.cpp
+++ b/Engine/Plugins/Media/D3D12VideoDecodersElectra/Source/Private/VideoDecoder_D3D12_H265.cpp
@@ -3,6 +3,7 @@
 #include "VideoDecoder_D3D12_H265.h"
 #include "D3D12VideoDecodersElectraModule.h"
 #include "Utils/MPEG/ElectraBitstreamProcessor_H265.h"
+#include "PixelFormat.h"
 
 
 namespace ElectraVideoDecodersD3D12Video
@@ -177,7 +178,7 @@ IElectraDecoder::EDecoderError FD3D12VideoDecoder_H265::DecodeAccessUnit(const F
 			SliceInfo.NalUnitType = (uint8)NalUnitType;
 			SliceInfo.NuhLayerId = (uint8)NuhLayerId;
 			SliceInfo.NumTemporalIdPlus1 = (uint8)NumTemporalIdPlus1;
-			ElectraDecodersUtil::MPEG::H265::FBitstreamReader br;
+			ElectraDecodersUtil::MPEG::H265::FBitstreamReaderH265 br;
 			TUniquePtr<ElectraDecodersUtil::MPEG::H265::FRBSP> SliceRBSP;
 			if (!ElectraDecodersUtil::MPEG::H265::ParseSliceHeader(SliceRBSP, br, SliceInfo.Header, BitstreamParamsH265.VPSs, BitstreamParamsH265.SPSs, BitstreamParamsH265.PPSs, Start+4, NaluLength))
 			{
@@ -782,17 +783,17 @@ IElectraDecoder::EDecoderError FD3D12VideoDecoder_H265::DecodeSlicesH265(const F
 	InDec->NumBits = InSequenceParameterSet.bit_depth_luma_minus8 + 8;
 	if (InDec->NumBits == 8)
 	{
-		InDec->BufferFormat = EElectraDecoderPlatformPixelFormat::NV12;
-		InDec->BufferEncoding = EElectraDecoderPlatformPixelEncoding::Native;
-		InDec->ExtraValues.Emplace(TEXT("pixfmt"), FVariant((int64)EElectraDecoderPlatformPixelFormat::NV12));
-		InDec->ExtraValues.Emplace(TEXT("pixenc"), FVariant((int64)EElectraDecoderPlatformPixelEncoding::Native));
+		InDec->BufferFormat = EPixelFormat::PF_NV12;
+		InDec->BufferEncoding = EElectraTextureSamplePixelEncoding::Native;
+		InDec->ExtraValues.Emplace(TEXT("pixfmt"), FVariant((int64)EPixelFormat::PF_NV12));
+		InDec->ExtraValues.Emplace(TEXT("pixenc"), FVariant((int64)EElectraTextureSamplePixelEncoding::Native));
 	}
 	else
 	{
-		InDec->BufferFormat = EElectraDecoderPlatformPixelFormat::P010;
-		InDec->BufferEncoding = EElectraDecoderPlatformPixelEncoding::Native;
-		InDec->ExtraValues.Emplace(TEXT("pixfmt"), FVariant((int64)EElectraDecoderPlatformPixelFormat::P010));
-		InDec->ExtraValues.Emplace(TEXT("pixenc"), FVariant((int64)EElectraDecoderPlatformPixelEncoding::Native));
+		InDec->BufferFormat = EPixelFormat::PF_P010;
+		InDec->BufferEncoding = EElectraTextureSamplePixelEncoding::Native;
+		InDec->ExtraValues.Emplace(TEXT("pixfmt"), FVariant((int64)EPixelFormat::PF_P010));
+		InDec->ExtraValues.Emplace(TEXT("pixenc"), FVariant((int64)EElectraTextureSamplePixelEncoding::Native));
 	}
 	InSequenceParameterSet.GetAspect(InDec->AspectW, InDec->AspectH);
 	auto FrameRate = InSequenceParameterSet.GetTiming();
diff --git a/Engine/Plugins/Media/ElectraCodecs/Source/ElectraCodecFactory/Public/IElectraCodecFactory.h b/Engine/Plugins/Media/ElectraCodecs/Source/ElectraCodecFactory/Public/IElectraCodecFactory.h
--- a/Engine/Plugins/Media/ElectraCodecs/Source/ElectraCodecFactory/Public/IElectraCodecFactory.h
+++ b/Engine/Plugins/Media/ElectraCodecs/Source/ElectraCodecFactory/Public/IElectraCodecFactory.h
@@ -7,7 +7,6 @@
 #include "Containers/Map.h"
 #include "Misc/Variant.h"
 
-#include "IElectraDecoderResourceDelegate.h"
 class IElectraDecoder;
 
 
diff --git a/Engine/Plugins/Media/ElectraCodecs/Source/ElectraDecoders/ElectraDecoders.Build.cs b/Engine/Plugins/Media/ElectraCodecs/Source/ElectraDecoders/ElectraDecoders.Build.cs
--- a/Engine/Plugins/Media/ElectraCodecs/Source/ElectraDecoders/ElectraDecoders.Build.cs
+++ b/Engine/Plugins/Media/ElectraCodecs/Source/ElectraDecoders/ElectraDecoders.Build.cs
@@ -1,6 +1,7 @@
 // Copyright Epic Games, Inc. All Rights Reserved.
 using UnrealBuildTool;
 using System.IO;
+using EpicGames.Core;
 
 namespace UnrealBuildTool.Rules
 {
@@ -28,18 +29,13 @@ namespace UnrealBuildTool.Rules
 				new string[] {
 					"Core",
                     "RHI",
+                    "Media",
 					"MediaUtils",
                     "ElectraSamples",
 					"ElectraBase",
                     // Common codecs
 					"dr_libs"					// MPEG audio layer 3
 				});
-
-			PrivateIncludePathModuleNames.AddRange(
-				new string[] {
-					"Media",
-				});
-
 			PublicIncludePathModuleNames.AddRange(
 				new string[] {
 					"ElectraCodecFactory"
@@ -55,7 +51,10 @@ namespace UnrealBuildTool.Rules
     			PrivateDependencyModuleNames.Add("SignalProcessing");
 
 				PublicDependencyModuleNames.Add("DirectX");
-				AddEngineThirdPartyPrivateStaticDependencies(Target, "DX12");
+
+				PrivateDependencyModuleNames.AddAll("D3D11RHI", "D3D12RHI");
+				AddEngineThirdPartyPrivateStaticDependencies(Target, "DX11", "DX12");
+
 
 				PrivateDefinitions.Add("_CRT_SECURE_NO_WARNINGS=1");
 
diff --git a/Engine/Plugins/Media/ElectraCodecs/Source/ElectraDecoders/Private/Android/AndroidElectraDecodersPlatformResources.cpp b/Engine/Plugins/Media/ElectraCodecs/Source/ElectraDecoders/Private/Android/AndroidElectraDecodersPlatformResources.cpp
--- a/Engine/Plugins/Media/ElectraCodecs/Source/ElectraDecoders/Private/Android/AndroidElectraDecodersPlatformResources.cpp
+++ b/Engine/Plugins/Media/ElectraCodecs/Source/ElectraDecoders/Private/Android/AndroidElectraDecodersPlatformResources.cpp
@@ -2,6 +2,12 @@
 
 #include "Android/AndroidElectraDecodersPlatformResources.h"
 #include "Misc/ScopeLock.h"
+
+#include "ElectraDecodersModule.h"
+#include "IElectraDecoderOutputVideo.h"
+#include "ElectraDecodersUtils.h"
+#include "MediaDecoderOutput.h"
+#include "MediaVideoDecoderOutput.h"
 #include "ElectraTextureSample.h"
 
 #include <atomic>
@@ -83,3 +89,67 @@ void FElectraDecodersPlatformResourcesAndroid::ReleaseResourceInstance(IDecoderR
 	ElectraCodecResourcesAndroid::FDecoderResourceInstance* Inst = static_cast<ElectraCodecResourcesAndroid::FDecoderResourceInstance*>(InInstanceToRelease);
 	delete Inst;
 }
+
+
+
+FElectraDecodersPlatformResourcesAndroid::IDecoderPlatformResource* FElectraDecodersPlatformResourcesAndroid::CreatePlatformVideoResource(void* InOwnerHandle, const TMap<FString, FVariant> InOptions)
+{
+	return nullptr;
+}
+
+void FElectraDecodersPlatformResourcesAndroid::ReleasePlatformVideoResource(void* InOwnerHandle, IDecoderPlatformResource* InHandleToRelease)
+{
+}
+
+bool FElectraDecodersPlatformResourcesAndroid::SetupOutputTextureSample(FString& OutErrorMessage, TSharedPtr<FElectraTextureSample, ESPMode::ThreadSafe> InOutTextureSampleToSetup, TSharedPtr<IElectraDecoderVideoOutput, ESPMode::ThreadSafe> InDecoderOutput, TSharedPtr<Electra::FParamDict, ESPMode::ThreadSafe> InOutBufferPropertes, IDecoderPlatformResource* InPlatformResource)
+{
+	if (!InOutTextureSampleToSetup.IsValid() || !InDecoderOutput.IsValid() || !InOutBufferPropertes.IsValid())// || !InPlatformResource)
+	{
+		OutErrorMessage = TEXT("Bad parameter");
+		return false;
+	}
+
+	TMap<FString, FVariant> ExtraValues;
+	InDecoderOutput->GetExtraValues(ExtraValues);
+
+	FElectraVideoDecoderOutputCropValues Crop = InDecoderOutput->GetCropValues();
+// TBD: Zero out the crop values for now since the output only consists of active pixels
+//      here and we don't have to crop ourselves.
+	Crop.Top = Crop.Left = Crop.Bottom = Crop.Right = 0;
+	InOutBufferPropertes->Set(IDecoderOutputOptionNames::Width, Electra::FVariantValue((int64)InDecoderOutput->GetWidth()));
+	InOutBufferPropertes->Set(IDecoderOutputOptionNames::Height, Electra::FVariantValue((int64)InDecoderOutput->GetHeight()));
+	InOutBufferPropertes->Set(IDecoderOutputOptionNames::CropLeft, Electra::FVariantValue((int64)Crop.Left));
+	InOutBufferPropertes->Set(IDecoderOutputOptionNames::CropRight, Electra::FVariantValue((int64)Crop.Right));
+	InOutBufferPropertes->Set(IDecoderOutputOptionNames::CropTop, Electra::FVariantValue((int64)Crop.Top));
+	InOutBufferPropertes->Set(IDecoderOutputOptionNames::CropBottom, Electra::FVariantValue((int64)Crop.Bottom));
+	InOutBufferPropertes->Set(IDecoderOutputOptionNames::AspectRatio, Electra::FVariantValue((double)InDecoderOutput->GetAspectRatioW() / (double)InDecoderOutput->GetAspectRatioH()));
+	InOutBufferPropertes->Set(IDecoderOutputOptionNames::AspectW, Electra::FVariantValue((int64)InDecoderOutput->GetAspectRatioW()));
+	InOutBufferPropertes->Set(IDecoderOutputOptionNames::AspectH, Electra::FVariantValue((int64)InDecoderOutput->GetAspectRatioH()));
+	InOutBufferPropertes->Set(IDecoderOutputOptionNames::FPSNumerator, Electra::FVariantValue((int64)InDecoderOutput->GetFrameRateNumerator()));
+	InOutBufferPropertes->Set(IDecoderOutputOptionNames::FPSDenominator, Electra::FVariantValue((int64)InDecoderOutput->GetFrameRateDenominator()));
+	InOutBufferPropertes->Set(IDecoderOutputOptionNames::PixelEncoding, Electra::FVariantValue((int64)EElectraTextureSamplePixelEncoding::Native));
+	InOutBufferPropertes->Set(IDecoderOutputOptionNames::PixelDataScale, Electra::FVariantValue((double)ElectraDecodersUtil::GetVariantValueSafeDouble(ExtraValues, TEXT("pix_datascale"), 1.0)));
+
+	int32 NumBits = InDecoderOutput->GetNumberOfBits();
+	EPixelFormat PixFmt = static_cast<EPixelFormat>(ElectraDecodersUtil::GetVariantValueSafeI64(ExtraValues, TEXT("pixfmt"), (int64)((NumBits > 8)  ? EPixelFormat::PF_A2B10G10R10 : EPixelFormat::PF_B8G8R8A8)));
+	int32 Pitch = static_cast<int32>(ElectraDecodersUtil::GetVariantValueSafeI64(ExtraValues, TEXT("pitch"), (int64)InDecoderOutput->GetDecodedWidth() * ((NumBits > 8) ? 2 : 1)));
+
+	InOutBufferPropertes->Set(IDecoderOutputOptionNames::PixelFormat, Electra::FVariantValue((int64)PixFmt));
+	InOutBufferPropertes->Set(IDecoderOutputOptionNames::Pitch, Electra::FVariantValue((int64)Pitch));
+	InOutBufferPropertes->Set(IDecoderOutputOptionNames::BitsPerComponent, Electra::FVariantValue((int64)NumBits));
+
+
+	InOutTextureSampleToSetup->InitializeCommon(*InOutBufferPropertes);
+	if (!InOutTextureSampleToSetup->Initialize())
+	{
+		OutErrorMessage = FString::Printf(TEXT("SetupOutputTextureSampleFromBuffer(): Failed to initialize texture sample"));
+		return false;
+	}
+	if (!InOutTextureSampleToSetup->FinishInitialization())
+	{
+		OutErrorMessage = FString::Printf(TEXT("SetupOutputTextureSampleFromBuffer(): Failed to finalize texture sample"));
+		return false;
+	}
+	return true;
+}
+
diff --git a/Engine/Plugins/Media/ElectraCodecs/Source/ElectraDecoders/Private/Android/aac/AudioDecoderAAC_Android.h b/Engine/Plugins/Media/ElectraCodecs/Source/ElectraDecoders/Private/Android/aac/AudioDecoderAAC_Android.h
--- a/Engine/Plugins/Media/ElectraCodecs/Source/ElectraDecoders/Private/Android/aac/AudioDecoderAAC_Android.h
+++ b/Engine/Plugins/Media/ElectraCodecs/Source/ElectraDecoders/Private/Android/aac/AudioDecoderAAC_Android.h
@@ -4,7 +4,6 @@
 
 
 #include "IElectraDecoder.h"
-#include "IElectraDecoderResourceDelegate.h"
 
 class IElectraAudioDecoderAAC_Android : public IElectraDecoder
 {
diff --git a/Engine/Plugins/Media/ElectraCodecs/Source/ElectraDecoders/Private/Apple/AppleElectraDecodersPlatformResources.cpp b/Engine/Plugins/Media/ElectraCodecs/Source/ElectraDecoders/Private/Apple/AppleElectraDecodersPlatformResources.cpp
--- a/Engine/Plugins/Media/ElectraCodecs/Source/ElectraDecoders/Private/Apple/AppleElectraDecodersPlatformResources.cpp
+++ b/Engine/Plugins/Media/ElectraCodecs/Source/ElectraDecoders/Private/Apple/AppleElectraDecodersPlatformResources.cpp
@@ -4,6 +4,14 @@
 #include "Templates/SharedPointer.h"
 #include "RHI.h"
 
+#include "ElectraDecodersModule.h"
+#include "IElectraDecoderOutputVideo.h"
+#include "ElectraDecodersUtils.h"
+#include "MediaDecoderOutput.h"
+#include "MediaVideoDecoderOutput.h"
+#include "ElectraTextureSample.h"
+
+#if PLATFORM_MAC || PLATFORM_IOS || PLATFORM_TVOS
 
 bool FElectraDecodersPlatformResourcesApple::GetMetalDevice(void** OutMetalDevice)
 {
@@ -20,3 +28,102 @@ bool FElectraDecodersPlatformResourcesApple::GetMetalDevice(void** OutMetalDevic
 	*OutMetalDevice = GDynamicRHI->RHIGetNativeDevice();
 	return true;
 }
+
+FElectraDecodersPlatformResourcesApple::IDecoderPlatformResource* FElectraDecodersPlatformResourcesApple::CreatePlatformVideoResource(void* InOwnerHandle, const TMap<FString, FVariant> InOptions)
+{
+	return nullptr;
+}
+
+void FElectraDecodersPlatformResourcesApple::ReleasePlatformVideoResource(void* InOwnerHandle, FElectraDecodersPlatformResourcesApple::IDecoderPlatformResource* InHandleToRelease)
+{
+}
+
+bool FElectraDecodersPlatformResourcesApple::SetupOutputTextureSample(FString& OutErrorMessage, TSharedPtr<FElectraTextureSample, ESPMode::ThreadSafe> InOutTextureSampleToSetup, TSharedPtr<IElectraDecoderVideoOutput, ESPMode::ThreadSafe> InDecoderOutput, TSharedPtr<Electra::FParamDict, ESPMode::ThreadSafe> InOutBufferPropertes, FElectraDecodersPlatformResourcesApple::IDecoderPlatformResource* InPlatformResource)
+{
+	if (!InOutTextureSampleToSetup.IsValid() || !InDecoderOutput.IsValid() || !InOutBufferPropertes.IsValid())// || !InPlatformResource)
+	{
+		OutErrorMessage = TEXT("Bad parameter");
+		return false;
+	}
+
+	TMap<FString, FVariant> ExtraValues;
+	InDecoderOutput->GetExtraValues(ExtraValues);
+
+	FElectraVideoDecoderOutputCropValues Crop = InDecoderOutput->GetCropValues();
+	InOutBufferPropertes->Set(IDecoderOutputOptionNames::Width, Electra::FVariantValue((int64)InDecoderOutput->GetWidth()));
+	InOutBufferPropertes->Set(IDecoderOutputOptionNames::Height, Electra::FVariantValue((int64)InDecoderOutput->GetHeight()));
+	InOutBufferPropertes->Set(IDecoderOutputOptionNames::CropLeft, Electra::FVariantValue((int64)Crop.Left));
+	InOutBufferPropertes->Set(IDecoderOutputOptionNames::CropRight, Electra::FVariantValue((int64)Crop.Right));
+	InOutBufferPropertes->Set(IDecoderOutputOptionNames::CropTop, Electra::FVariantValue((int64)Crop.Top));
+	InOutBufferPropertes->Set(IDecoderOutputOptionNames::CropBottom, Electra::FVariantValue((int64)Crop.Bottom));
+	InOutBufferPropertes->Set(IDecoderOutputOptionNames::AspectRatio, Electra::FVariantValue((double)InDecoderOutput->GetAspectRatioW() / (double)InDecoderOutput->GetAspectRatioH()));
+	InOutBufferPropertes->Set(IDecoderOutputOptionNames::AspectW, Electra::FVariantValue((int64)InDecoderOutput->GetAspectRatioW()));
+	InOutBufferPropertes->Set(IDecoderOutputOptionNames::AspectH, Electra::FVariantValue((int64)InDecoderOutput->GetAspectRatioH()));
+	InOutBufferPropertes->Set(IDecoderOutputOptionNames::FPSNumerator, Electra::FVariantValue((int64)InDecoderOutput->GetFrameRateNumerator()));
+	InOutBufferPropertes->Set(IDecoderOutputOptionNames::FPSDenominator, Electra::FVariantValue((int64)InDecoderOutput->GetFrameRateDenominator()));
+	InOutBufferPropertes->Set(IDecoderOutputOptionNames::BitsPerComponent, Electra::FVariantValue((int64)InDecoderOutput->GetNumberOfBits()));
+	InOutBufferPropertes->Set(IDecoderOutputOptionNames::PixelDataScale, Electra::FVariantValue((double)ElectraDecodersUtil::GetVariantValueSafeDouble(ExtraValues, TEXT("pix_datascale"), 1.0)));
+
+	IElectraDecoderVideoOutputImageBuffers* ImageBuffers = reinterpret_cast<IElectraDecoderVideoOutputImageBuffers*>(InDecoderOutput->GetPlatformOutputHandle(EElectraDecoderPlatformOutputHandleType::ImageBuffers));
+	if (ImageBuffers && ImageBuffers->GetNumberOfBuffers() == 1)
+	{
+		EPixelFormat PixFmt = ImageBuffers->GetBufferFormatByIndex(0);
+		EElectraTextureSamplePixelEncoding PixEnc = ImageBuffers->GetBufferEncodingByIndex(0);
+		int32 Pitch = ImageBuffers->GetBufferPitchByIndex(0);
+		InOutBufferPropertes->Set(IDecoderOutputOptionNames::PixelFormat, Electra::FVariantValue((int64)PixFmt));
+		InOutBufferPropertes->Set(IDecoderOutputOptionNames::PixelEncoding, Electra::FVariantValue((int64)PixEnc));
+		InOutBufferPropertes->Set(IDecoderOutputOptionNames::Pitch, Electra::FVariantValue((int64)Pitch));
+
+		InOutTextureSampleToSetup->InitializeCommon(*InOutBufferPropertes);
+
+		bool bIsValid = false;
+		TSharedPtr<TArray64<uint8>, ESPMode::ThreadSafe> ColorBuffer = ImageBuffers->GetBufferDataByIndex(0);
+		if (ColorBuffer.IsValid() && ColorBuffer->Num())
+		{
+			//
+			// CPU side buffer
+			//
+			int32 Width = InDecoderOutput->GetDecodedWidth();
+			int32 Height = InDecoderOutput->GetDecodedHeight();
+
+			InOutTextureSampleToSetup->Buffer = MoveTemp(ColorBuffer);
+			InOutTextureSampleToSetup->Stride = Pitch;
+			InOutTextureSampleToSetup->SetDim(FIntPoint(Width, Height));
+			bIsValid = true;
+		}
+		else
+		{
+			//
+			// GPU texture (CVImageBuffer)
+			//
+			CVImageBufferRef ImageBuffer = static_cast<CVImageBufferRef>(ImageBuffers->GetBufferTextureByIndex(0));
+			if (ImageBuffer != nullptr)
+			{
+				if (InOutTextureSampleToSetup->ImageBufferRef)
+				{
+					CFRelease(InOutTextureSampleToSetup->ImageBufferRef);
+				}
+				InOutTextureSampleToSetup->ImageBufferRef = ImageBuffer;
+				CFRetain(InOutTextureSampleToSetup->ImageBufferRef);
+				InOutTextureSampleToSetup->Stride = CVPixelBufferGetBytesPerRow(InOutTextureSampleToSetup->ImageBufferRef);
+				bIsValid = true;
+			}
+		}
+		// Finish initialization.
+		if (bIsValid)
+		{
+			if (!InOutTextureSampleToSetup->FinishInitialization())
+			{
+				OutErrorMessage = FString::Printf(TEXT("SetupOutputTextureSampleFromBuffer(): Unsupported pixel format encoding"));
+				return false;
+			}
+			return true;
+		}
+	}
+	OutErrorMessage = TEXT("Unhandled decoder output format");
+	return false;
+}
+
+
+
+#endif // PLATFORM_MAC || PLATFORM_IOS || PLATFORM_TVOS
diff --git a/Engine/Plugins/Media/ElectraCodecs/Source/ElectraDecoders/Private/Apple/h264/VideoDecoderH264_Apple.cpp b/Engine/Plugins/Media/ElectraCodecs/Source/ElectraDecoders/Private/Apple/h264/VideoDecoderH264_Apple.cpp
--- a/Engine/Plugins/Media/ElectraCodecs/Source/ElectraDecoders/Private/Apple/h264/VideoDecoderH264_Apple.cpp
+++ b/Engine/Plugins/Media/ElectraCodecs/Source/ElectraDecoders/Private/Apple/h264/VideoDecoderH264_Apple.cpp
@@ -13,6 +13,7 @@
 #include "IElectraDecoderOutputVideo.h"
 
 #include "ElectraDecodersModule.h"
+#include "PixelFormat.h"
 
 #include "Apple/AppleElectraDecoderPlatformOutputFormatTypes.h"
 
@@ -80,7 +81,7 @@ public:
 	{
 		return 1;
 	}
-	TSharedPtr<TArray<uint8>, ESPMode::ThreadSafe> GetBufferDataByIndex(int32 InBufferIndex) const override
+	TSharedPtr<TArray64<uint8>, ESPMode::ThreadSafe> GetBufferDataByIndex(int32 InBufferIndex) const override
 	{
 		return nullptr;
 	}
@@ -92,11 +93,11 @@ public:
 	{
 		return false;
 	}
-	EElectraDecoderPlatformPixelFormat GetBufferFormatByIndex(int32 InBufferIndex) const override
+	EPixelFormat GetBufferFormatByIndex(int32 InBufferIndex) const override
 	{
 		return PixelFormat;
 	}
-	EElectraDecoderPlatformPixelEncoding GetBufferEncodingByIndex(int32 InBufferIndex) const override
+	EElectraTextureSamplePixelEncoding GetBufferEncodingByIndex(int32 InBufferIndex) const override
 	{
 		return PixelEncoding;
 	}
@@ -127,8 +128,8 @@ public:
 	int32 AspectH = 1;
 	int32 FrameRateN = 0;
 	int32 FrameRateD = 0;
-	EElectraDecoderPlatformPixelFormat PixelFormat = EElectraDecoderPlatformPixelFormat::INVALID;
-	EElectraDecoderPlatformPixelEncoding PixelEncoding = EElectraDecoderPlatformPixelEncoding::Native;
+	EPixelFormat PixelFormat = EPixelFormat::PF_Unknown;
+	EElectraTextureSamplePixelEncoding PixelEncoding = EElectraTextureSamplePixelEncoding::Native;
 	TMap<FString, FVariant> ExtraValues;
 
 	CVImageBufferRef ImageBuffer = nullptr;
@@ -882,27 +883,27 @@ FElectraVideoDecoderH264_Apple::EConvertResult FElectraVideoDecoderH264_Apple::C
 	NewOutput->ImageBuffer = NextImage.ReleaseImageBufferRef();
 	int pixelFormat = CVPixelBufferGetPixelFormatType(NewOutput->ImageBuffer);
 
-	switch (pixelFormat)
+	switch(pixelFormat)
 	{
-		case	kCVPixelFormatType_420YpCbCr8BiPlanarVideoRange:
-		case	kCVPixelFormatType_420YpCbCr8BiPlanarFullRange:
+		case kCVPixelFormatType_420YpCbCr8BiPlanarVideoRange:
+		case kCVPixelFormatType_420YpCbCr8BiPlanarFullRange:
 		{
-			NewOutput->PixelFormat = EElectraDecoderPlatformPixelFormat::NV12;
-			NewOutput->PixelEncoding = EElectraDecoderPlatformPixelEncoding::Native;
+			NewOutput->PixelFormat = EPixelFormat::PF_NV12;
+			NewOutput->PixelEncoding = EElectraTextureSamplePixelEncoding::Native;
 			break;
 		}
-		case	kCVPixelFormatType_420YpCbCr10BiPlanarVideoRange:
-		case	kCVPixelFormatType_420YpCbCr10BiPlanarFullRange:
+		case kCVPixelFormatType_420YpCbCr10BiPlanarVideoRange:
+		case kCVPixelFormatType_420YpCbCr10BiPlanarFullRange:
 		{
-			NewOutput->PixelFormat = EElectraDecoderPlatformPixelFormat::P010;
-			NewOutput->PixelEncoding = EElectraDecoderPlatformPixelEncoding::Native;
+			NewOutput->PixelFormat = EPixelFormat::PF_P010;
+			NewOutput->PixelEncoding = EElectraTextureSamplePixelEncoding::Native;
 			break;
 		}
 		default:
 		{
 			check(!"Unexpected output format!");
-			NewOutput->PixelFormat = EElectraDecoderPlatformPixelFormat::INVALID;
-			NewOutput->PixelEncoding = EElectraDecoderPlatformPixelEncoding::Native;
+			NewOutput->PixelFormat = EPixelFormat::PF_Unknown;
+			NewOutput->PixelEncoding = EElectraTextureSamplePixelEncoding::Native;
 		}
 	}
 
diff --git a/Engine/Plugins/Media/ElectraCodecs/Source/ElectraDecoders/Private/Apple/h265/VideoDecoderH265_Apple.cpp b/Engine/Plugins/Media/ElectraCodecs/Source/ElectraDecoders/Private/Apple/h265/VideoDecoderH265_Apple.cpp
--- a/Engine/Plugins/Media/ElectraCodecs/Source/ElectraDecoders/Private/Apple/h265/VideoDecoderH265_Apple.cpp
+++ b/Engine/Plugins/Media/ElectraCodecs/Source/ElectraDecoders/Private/Apple/h265/VideoDecoderH265_Apple.cpp
@@ -13,6 +13,7 @@
 #include "IElectraDecoderOutputVideo.h"
 
 #include "ElectraDecodersModule.h"
+#include "PixelFormat.h"
 
 #include "Apple/AppleElectraDecoderPlatformOutputFormatTypes.h"
 
@@ -80,7 +81,7 @@ public:
 	{
 		return 1;
 	}
-	TSharedPtr<TArray<uint8>, ESPMode::ThreadSafe> GetBufferDataByIndex(int32 InBufferIndex) const override
+	TSharedPtr<TArray64<uint8>, ESPMode::ThreadSafe> GetBufferDataByIndex(int32 InBufferIndex) const override
 	{
 		return nullptr;
 	}
@@ -92,11 +93,11 @@ public:
 	{
 		return false;
 	}
-	EElectraDecoderPlatformPixelFormat GetBufferFormatByIndex(int32 InBufferIndex) const override
+	EPixelFormat GetBufferFormatByIndex(int32 InBufferIndex) const override
 	{
 		return PixelFormat;
 	}
-	EElectraDecoderPlatformPixelEncoding GetBufferEncodingByIndex(int32 InBufferIndex) const override
+	EElectraTextureSamplePixelEncoding GetBufferEncodingByIndex(int32 InBufferIndex) const override
 	{
 		return PixelEncoding;
 	}
@@ -127,8 +128,8 @@ public:
 	int32 AspectH = 1;
 	int32 FrameRateN = 0;
 	int32 FrameRateD = 0;
-	EElectraDecoderPlatformPixelFormat PixelFormat = EElectraDecoderPlatformPixelFormat::INVALID;
-	EElectraDecoderPlatformPixelEncoding PixelEncoding = EElectraDecoderPlatformPixelEncoding::Native;
+	EPixelFormat PixelFormat = EPixelFormat::PF_Unknown;
+	EElectraTextureSamplePixelEncoding PixelEncoding = EElectraTextureSamplePixelEncoding::Native;
 	TMap<FString, FVariant> ExtraValues;
 
 	CVImageBufferRef ImageBuffer = nullptr;
@@ -885,27 +886,27 @@ FElectraVideoDecoderH265_Apple::EConvertResult FElectraVideoDecoderH265_Apple::C
 	NewOutput->ImageBuffer = NextImage.ReleaseImageBufferRef();
 	int pixelFormat = CVPixelBufferGetPixelFormatType(NewOutput->ImageBuffer);
 
-	switch (pixelFormat)
+	switch(pixelFormat)
 	{
-		case	kCVPixelFormatType_420YpCbCr8BiPlanarVideoRange:
-		case	kCVPixelFormatType_420YpCbCr8BiPlanarFullRange:
+		case kCVPixelFormatType_420YpCbCr8BiPlanarVideoRange:
+		case kCVPixelFormatType_420YpCbCr8BiPlanarFullRange:
 		{
-			NewOutput->PixelFormat = EElectraDecoderPlatformPixelFormat::NV12;
-			NewOutput->PixelEncoding = EElectraDecoderPlatformPixelEncoding::Native;
+			NewOutput->PixelFormat = EPixelFormat::PF_NV12;
+			NewOutput->PixelEncoding = EElectraTextureSamplePixelEncoding::Native;
 			break;
 		}
-		case	kCVPixelFormatType_420YpCbCr10BiPlanarVideoRange:
-		case	kCVPixelFormatType_420YpCbCr10BiPlanarFullRange:
+		case kCVPixelFormatType_420YpCbCr10BiPlanarVideoRange:
+		case kCVPixelFormatType_420YpCbCr10BiPlanarFullRange:
 		{
-			NewOutput->PixelFormat = EElectraDecoderPlatformPixelFormat::P010;
-			NewOutput->PixelEncoding = EElectraDecoderPlatformPixelEncoding::Native;
+			NewOutput->PixelFormat = EPixelFormat::PF_P010;
+			NewOutput->PixelEncoding = EElectraTextureSamplePixelEncoding::Native;
 			break;
 		}
 		default:
 		{
 			check(!"Unexpected output format!");
-			NewOutput->PixelFormat = EElectraDecoderPlatformPixelFormat::INVALID;
-			NewOutput->PixelEncoding = EElectraDecoderPlatformPixelEncoding::Native;
+			NewOutput->PixelFormat = EPixelFormat::PF_Unknown;
+			NewOutput->PixelEncoding = EElectraTextureSamplePixelEncoding::Native;
 		}
 	}
 
diff --git a/Engine/Plugins/Media/ElectraCodecs/Source/ElectraDecoders/Private/DX/VideoDecoderH264_DX.cpp b/Engine/Plugins/Media/ElectraCodecs/Source/ElectraDecoders/Private/DX/VideoDecoderH264_DX.cpp
--- a/Engine/Plugins/Media/ElectraCodecs/Source/ElectraDecoders/Private/DX/VideoDecoderH264_DX.cpp
+++ b/Engine/Plugins/Media/ElectraCodecs/Source/ElectraDecoders/Private/DX/VideoDecoderH264_DX.cpp
@@ -13,16 +13,16 @@
 #include "IElectraDecoderOutputVideo.h"
 
 #include "ElectraDecodersModule.h"
+#include "PixelFormat.h"
 
 #include "ElectraDecodersPlatformResources.h"
-#include COMPILED_PLATFORM_HEADER(ElectraDecoderGPUBufferHelpers.h)
+#include COMPILED_PLATFORM_HEADER(ElectraTextureSampleGPUBufferHelper.h)
 #include COMPILED_PLATFORM_HEADER(PlatformHeaders_Video_DX.h)
 
 #include "Stats/Stats.h"
 #include "HAL/IConsoleManager.h"
 
 /*********************************************************************************************************************/
-#include COMPILED_PLATFORM_HEADER(PlatformHeaders_Video_DX.h)
 
 #define VERIFY_HR(FNcall, Msg, Code)	\
 Result = FNcall;						\
@@ -99,23 +99,26 @@ public:
 		{
 #if ELECTRA_HAVE_IMAGEBUFFERS
 			case EElectraDecoderPlatformOutputHandleType::ImageBuffers:
-				if (!MFSample.IsValid())
-				{
-					if (InTypeOfHandle == EElectraDecoderPlatformOutputHandleType::ImageBuffers)	//-V547
-					{
-						return static_cast<IElectraDecoderVideoOutputImageBuffers*>(const_cast<FElectraVideoDecoderOutputH264_DX*>(this));
-					}
-				}
-				return nullptr;
+			{
+				return !MFSample.IsValid() ? static_cast<IElectraDecoderVideoOutputImageBuffers*>(const_cast<FElectraVideoDecoderOutputH264_DX*>(this)) : nullptr;
+			}
 #endif
 			case EElectraDecoderPlatformOutputHandleType::MFSample:
+			{
 				return MFSample.GetReference();
+			}
 			case EElectraDecoderPlatformOutputHandleType::DXDevice:
+			{
 				return Ctx.DxDevice.GetReference();
+			}
 			case EElectraDecoderPlatformOutputHandleType::DXDeviceContext:
+			{
 				return Ctx.DxDeviceContext.GetReference();
+			}
 			default:
+			{
 				return nullptr;
+			}
 		}
 	}
 	IElectraDecoderVideoOutput::EImageCopyResult CopyPlatformImage(IElectraDecoderVideoOutputCopyResources* InCopyResources) const override
@@ -127,26 +130,26 @@ public:
 	//
 
 	// Return the 4cc of the codec.
-	virtual uint32 GetCodec4CC() const override
+	uint32 GetCodec4CC() const override
 	{
 		return *(uint32*)"h264";
 	}
 
 	// Returns the number of separate image buffers making up the frame.
-	virtual int32 GetNumberOfBuffers() const override
+	int32 GetNumberOfBuffers() const override
 	{
 		return 1;
 	}
 
 	// Returns the n'th image buffer as CPU data buffer.
-	virtual TSharedPtr<TArray<uint8>, ESPMode::ThreadSafe> GetBufferDataByIndex(int32 InBufferIndex) const override
+	TSharedPtr<TArray64<uint8>, ESPMode::ThreadSafe> GetBufferDataByIndex(int32 InBufferIndex) const override
 	{
 		// For now CPU side buffers use the old style MFSample way to return data
 		return nullptr;
 	}
 
 	// Returns the n'th image buffer as GPU texture resource.
-	virtual void* GetBufferTextureByIndex(int32 InBufferIndex) const override
+	void* GetBufferTextureByIndex(int32 InBufferIndex) const override
 	{
 		if (InBufferIndex == 0)
 		{
@@ -155,7 +158,7 @@ public:
 		return nullptr;
 	}
 
-	virtual bool GetBufferTextureSyncByIndex(int32 InBufferIndex, FElectraDecoderOutputSync& SyncObject) const override
+	bool GetBufferTextureSyncByIndex(int32 InBufferIndex, FElectraDecoderOutputSync& SyncObject) const override
 	{
 		if (InBufferIndex == 0)
 		{
@@ -171,27 +174,27 @@ public:
 	}
 
 	// Returns the n'th image buffer format
-	virtual EElectraDecoderPlatformPixelFormat GetBufferFormatByIndex(int32 InBufferIndex) const override
+	EPixelFormat GetBufferFormatByIndex(int32 InBufferIndex) const override
 	{
 		if (InBufferIndex == 0)
 		{
-			return (EElectraDecoderPlatformPixelFormat)ElectraDecodersUtil::GetVariantValueSafeU64(ExtraValues, TEXT("pixfmt"), (uint64)EElectraDecoderPlatformPixelFormat::INVALID);
+			return (EPixelFormat)ElectraDecodersUtil::GetVariantValueSafeU64(ExtraValues, TEXT("pixfmt"), (uint64)EPixelFormat::PF_Unknown);
 		}
-		return EElectraDecoderPlatformPixelFormat::INVALID;
+		return EPixelFormat::PF_Unknown;
 	}
 
 	// Returns the n'th image buffer encoding
-	virtual EElectraDecoderPlatformPixelEncoding GetBufferEncodingByIndex(int32 InBufferIndex) const override
+	EElectraTextureSamplePixelEncoding GetBufferEncodingByIndex(int32 InBufferIndex) const override
 	{
 		if (InBufferIndex == 0)
 		{
-			return (EElectraDecoderPlatformPixelEncoding)ElectraDecodersUtil::GetVariantValueSafeU64(ExtraValues, TEXT("pixenc"), (uint64)EElectraDecoderPlatformPixelEncoding::Native);
+			return (EElectraTextureSamplePixelEncoding)ElectraDecodersUtil::GetVariantValueSafeU64(ExtraValues, TEXT("pixenc"), (uint64)EElectraTextureSamplePixelEncoding::Native);
 		}
-		return EElectraDecoderPlatformPixelEncoding::Native;
+		return EElectraTextureSamplePixelEncoding::Native;
 	}
 
 	// Returns the n'th image buffer pitch
-	virtual int32 GetBufferPitchByIndex(int32 InBufferIndex) const override
+	int32 GetBufferPitchByIndex(int32 InBufferIndex) const override
 	{
 		if (InBufferIndex == 0)
 		{
@@ -1120,8 +1123,8 @@ bool FElectraVideoDecoderH264_DX::ConvertDecoderOutput()
 	NewOutput->ExtraValues.Emplace(TEXT("dxversion"), FVariant((int64) DecoderPlatformHandle->GetDXVersionTimes1000()));
 	NewOutput->ExtraValues.Emplace(TEXT("sw"), FVariant(DecoderPlatformHandle->IsSoftware()));
 	NewOutput->ExtraValues.Emplace(TEXT("codec"), FVariant(TEXT("avc")));
-	NewOutput->ExtraValues.Emplace(TEXT("pixfmt"), FVariant((int64)EElectraDecoderPlatformPixelFormat::NV12));
-	NewOutput->ExtraValues.Emplace(TEXT("pixenc"), FVariant((int64)EElectraDecoderPlatformPixelEncoding::Native));
+	NewOutput->ExtraValues.Emplace(TEXT("pixfmt"), FVariant((int64)EPixelFormat::PF_NV12));
+	NewOutput->ExtraValues.Emplace(TEXT("pixenc"), FVariant((int64)EElectraTextureSamplePixelEncoding::Native));
 
 #if ELECTRA_HAVE_IMAGEBUFFERS
 	void* PlatformDevice = nullptr;
diff --git a/Engine/Plugins/Media/ElectraCodecs/Source/ElectraDecoders/Private/DX/VideoDecoderH265_DX.cpp b/Engine/Plugins/Media/ElectraCodecs/Source/ElectraDecoders/Private/DX/VideoDecoderH265_DX.cpp
--- a/Engine/Plugins/Media/ElectraCodecs/Source/ElectraDecoders/Private/DX/VideoDecoderH265_DX.cpp
+++ b/Engine/Plugins/Media/ElectraCodecs/Source/ElectraDecoders/Private/DX/VideoDecoderH265_DX.cpp
@@ -12,16 +12,16 @@
 #include "IElectraDecoderOutputVideo.h"
 
 #include "ElectraDecodersModule.h"
+#include "PixelFormat.h"
 
 #include "ElectraDecodersPlatformResources.h"
-#include COMPILED_PLATFORM_HEADER(ElectraDecoderGPUBufferHelpers.h)
+#include COMPILED_PLATFORM_HEADER(ElectraTextureSampleGPUBufferHelper.h)
 #include COMPILED_PLATFORM_HEADER(PlatformHeaders_Video_DX.h)
 
 #include "Stats/Stats.h"
 #include "HAL/IConsoleManager.h"
 
 /*********************************************************************************************************************/
-#include COMPILED_PLATFORM_HEADER(PlatformHeaders_Video_DX.h)
 
 #define VERIFY_HR(FNcall, Msg, Code)	\
 Result = FNcall;						\
@@ -97,24 +97,27 @@ public:
 		switch (InTypeOfHandle)
 		{
 #if ELECTRA_HAVE_IMAGEBUFFERS
-		case EElectraDecoderPlatformOutputHandleType::ImageBuffers:
-			if (!MFSample.IsValid())
+			case EElectraDecoderPlatformOutputHandleType::ImageBuffers:
 			{
-				if (InTypeOfHandle == EElectraDecoderPlatformOutputHandleType::ImageBuffers)	//-V547
-				{
-					return static_cast<IElectraDecoderVideoOutputImageBuffers*>(const_cast<FElectraVideoDecoderOutputH265_DX*>(this));
-				}
+				return !MFSample.IsValid() ? static_cast<IElectraDecoderVideoOutputImageBuffers*>(const_cast<FElectraVideoDecoderOutputH265_DX*>(this)) : nullptr;
 			}
-			return nullptr;
 #endif
-		case EElectraDecoderPlatformOutputHandleType::MFSample:
-			return MFSample.GetReference();
-		case EElectraDecoderPlatformOutputHandleType::DXDevice:
-			return Ctx.DxDevice.GetReference();
-		case EElectraDecoderPlatformOutputHandleType::DXDeviceContext:
-			return Ctx.DxDeviceContext.GetReference();
-		default:
-			return nullptr;
+			case EElectraDecoderPlatformOutputHandleType::MFSample:
+			{
+				return MFSample.GetReference();
+			}
+			case EElectraDecoderPlatformOutputHandleType::DXDevice:
+			{
+				return Ctx.DxDevice.GetReference();
+			}
+			case EElectraDecoderPlatformOutputHandleType::DXDeviceContext:
+			{
+				return Ctx.DxDeviceContext.GetReference();
+			}
+			default:
+			{
+				return nullptr;
+			}
 		}
 	}
 	IElectraDecoderVideoOutput::EImageCopyResult CopyPlatformImage(IElectraDecoderVideoOutputCopyResources* InCopyResources) const override
@@ -126,26 +129,26 @@ public:
 	//
 
 	// Return the 4cc of the codec.
-	virtual uint32 GetCodec4CC() const override
+	uint32 GetCodec4CC() const override
 	{
 		return *(uint32*)"hevc";
 	}
 
 	// Returns the number of separate image buffers making up the frame.
-	virtual int32 GetNumberOfBuffers() const override
+	int32 GetNumberOfBuffers() const override
 	{
 		return 1;
 	}
 
 	// Returns the n'th image buffer as CPU data buffer.
-	virtual TSharedPtr<TArray<uint8>, ESPMode::ThreadSafe> GetBufferDataByIndex(int32 InBufferIndex) const override
+	TSharedPtr<TArray64<uint8>, ESPMode::ThreadSafe> GetBufferDataByIndex(int32 InBufferIndex) const override
 	{
 		// For now CPU side buffers use the old style MFSample way to return data
 		return nullptr;
 	}
 
 	// Returns the n'th image buffer as GPU texture resource.
-	virtual void* GetBufferTextureByIndex(int32 InBufferIndex) const override
+	void* GetBufferTextureByIndex(int32 InBufferIndex) const override
 	{
 		if (InBufferIndex == 0)
 		{
@@ -154,7 +157,7 @@ public:
 		return nullptr;
 	}
 
-	virtual bool GetBufferTextureSyncByIndex(int32 InBufferIndex, FElectraDecoderOutputSync& SyncObject) const override
+	bool GetBufferTextureSyncByIndex(int32 InBufferIndex, FElectraDecoderOutputSync& SyncObject) const override
 	{
 		if (InBufferIndex == 0)
 		{
@@ -170,27 +173,27 @@ public:
 	}
 
 	// Returns the n'th image buffer format
-	virtual EElectraDecoderPlatformPixelFormat GetBufferFormatByIndex(int32 InBufferIndex) const override
+	EPixelFormat GetBufferFormatByIndex(int32 InBufferIndex) const override
 	{
 		if (InBufferIndex == 0)
 		{
-			return (EElectraDecoderPlatformPixelFormat)ElectraDecodersUtil::GetVariantValueSafeU64(ExtraValues, TEXT("pixfmt"), (uint64)EElectraDecoderPlatformPixelFormat::INVALID);
+			return (EPixelFormat)ElectraDecodersUtil::GetVariantValueSafeU64(ExtraValues, TEXT("pixfmt"), (uint64)EPixelFormat::PF_Unknown);
 		}
-		return EElectraDecoderPlatformPixelFormat::INVALID;
+		return EPixelFormat::PF_Unknown;
 	}
 
 	// Returns the n'th image buffer encoding
-	virtual EElectraDecoderPlatformPixelEncoding GetBufferEncodingByIndex(int32 InBufferIndex) const override
+	EElectraTextureSamplePixelEncoding GetBufferEncodingByIndex(int32 InBufferIndex) const override
 	{
 		if (InBufferIndex == 0)
 		{
-			return (EElectraDecoderPlatformPixelEncoding)ElectraDecodersUtil::GetVariantValueSafeU64(ExtraValues, TEXT("pixenc"), (uint64)EElectraDecoderPlatformPixelEncoding::Native);
+			return (EElectraTextureSamplePixelEncoding)ElectraDecodersUtil::GetVariantValueSafeU64(ExtraValues, TEXT("pixenc"), (uint64)EElectraTextureSamplePixelEncoding::Native);
 		}
-		return EElectraDecoderPlatformPixelEncoding::Native;
+		return EElectraTextureSamplePixelEncoding::Native;
 	}
 
 	// Returns the n'th image buffer pitch
-	virtual int32 GetBufferPitchByIndex(int32 InBufferIndex) const override
+	int32 GetBufferPitchByIndex(int32 InBufferIndex) const override
 	{
 		if (InBufferIndex == 0)
 		{
@@ -1101,7 +1104,7 @@ bool FElectraVideoDecoderH265_DX::ConvertDecoderOutput()
 
 	// Get actual output format used right now.
 	GUID CurrentVideoFormatGUID;
-	EElectraDecoderPlatformPixelFormat PixFmt;
+	EPixelFormat PixFmt;
 	if (SUCCEEDED(CurrentOutputMediaType->GetGUID(MF_MT_SUBTYPE, &CurrentVideoFormatGUID)))
 	{
 		check(CurrentVideoFormatGUID == MFVideoFormat_NV12 || CurrentVideoFormatGUID == MFVideoFormat_P010);
@@ -1114,20 +1117,20 @@ bool FElectraVideoDecoderH265_DX::ConvertDecoderOutput()
 	{
 		NewOutput->NumBits = 8;
 		NewOutput->PixelFormat = 0;
-		PixFmt = EElectraDecoderPlatformPixelFormat::NV12;
+		PixFmt = EPixelFormat::PF_NV12;
 	}
 	else
 	{
 		NewOutput->NumBits = 10;
 		NewOutput->PixelFormat = 1;
-		PixFmt = EElectraDecoderPlatformPixelFormat::P010;
+		PixFmt = EPixelFormat::PF_P010;
 	}
 	NewOutput->ExtraValues.Emplace(TEXT("platform"), FVariant(TEXT("dx")));
 	NewOutput->ExtraValues.Emplace(TEXT("dxversion"), FVariant((int64)DecoderPlatformHandle->GetDXVersionTimes1000()));
 	NewOutput->ExtraValues.Emplace(TEXT("sw"), FVariant(DecoderPlatformHandle->IsSoftware()));
 	NewOutput->ExtraValues.Emplace(TEXT("codec"), FVariant(TEXT("hevc")));
 	NewOutput->ExtraValues.Emplace(TEXT("pixfmt"), FVariant((int64)PixFmt));
-	NewOutput->ExtraValues.Emplace(TEXT("pixenc"), FVariant((int64)EElectraDecoderPlatformPixelEncoding::Native));
+	NewOutput->ExtraValues.Emplace(TEXT("pixenc"), FVariant((int64)EElectraTextureSamplePixelEncoding::Native));
 
 #if ELECTRA_HAVE_IMAGEBUFFERS
 	void* PlatformDevice = nullptr;
@@ -1164,9 +1167,9 @@ bool FElectraVideoDecoderH265_DX::ConvertDecoderOutput()
 		TRefCountPtr D3D12Device(static_cast<ID3D12Device*>(PlatformDevice));
 
 		// Create the resource pool as needed...
-		if (!D3D12ResourcePool || !D3D12ResourcePool->IsCompatibleAsBuffer(MaxOutputBuffers, MaxWidth, MaxHeight * 3 / 2, (PixFmt == EElectraDecoderPlatformPixelFormat::NV12) ? 1 : 2))
+		if (!D3D12ResourcePool || !D3D12ResourcePool->IsCompatibleAsBuffer(MaxOutputBuffers, MaxWidth, MaxHeight * 3 / 2, (PixFmt == EPixelFormat::PF_NV12) ? 1 : 2))
 		{
-			D3D12ResourcePool = MakeShared<FElectraMediaDecoderOutputBufferPool_DX12, ESPMode::ThreadSafe>(D3D12Device, MaxOutputBuffers, MaxWidth, MaxHeight * 3 / 2, (PixFmt == EElectraDecoderPlatformPixelFormat::NV12) ? 1 : 2);
+			D3D12ResourcePool = MakeShared<FElectraMediaDecoderOutputBufferPool_DX12, ESPMode::ThreadSafe>(D3D12Device, MaxOutputBuffers, MaxWidth, MaxHeight * 3 / 2, (PixFmt == EPixelFormat::PF_NV12) ? 1 : 2);
 		}
 
 		// Create a tasksync instance so we can have our own async jobs run in consecutive order
diff --git a/Engine/Plugins/Media/ElectraCodecs/Source/ElectraDecoders/Private/ElectraMediaMP3Decoder.cpp b/Engine/Plugins/Media/ElectraCodecs/Source/ElectraDecoders/Private/ElectraMediaMP3Decoder.cpp
--- a/Engine/Plugins/Media/ElectraCodecs/Source/ElectraDecoders/Private/ElectraMediaMP3Decoder.cpp
+++ b/Engine/Plugins/Media/ElectraCodecs/Source/ElectraDecoders/Private/ElectraMediaMP3Decoder.cpp
@@ -12,7 +12,7 @@
 #include "IElectraDecoder.h"
 #include "IElectraDecoderOutputAudio.h"
 #include "ElectraDecodersUtils.h"
-#include "Utils/ElectraBitstreamReader.h"
+#include "Utilities/ElectraBitstream.h"
 #include "Utils/MPEG/ElectraUtilsMPEGAudio.h"
 #include "Utils/ElectraBitstreamProcessorDefault.h"
 
diff --git a/Engine/Plugins/Media/ElectraCodecs/Source/ElectraDecoders/Private/Linux/LinuxElectraDecodersPlatformResources.cpp b/Engine/Plugins/Media/ElectraCodecs/Source/ElectraDecoders/Private/Linux/LinuxElectraDecodersPlatformResources.cpp
--- a/Engine/Plugins/Media/ElectraCodecs/Source/ElectraDecoders/Private/Linux/LinuxElectraDecodersPlatformResources.cpp
+++ b/Engine/Plugins/Media/ElectraCodecs/Source/ElectraDecoders/Private/Linux/LinuxElectraDecodersPlatformResources.cpp
@@ -1,10 +1,375 @@
 // Copyright Epic Games, Inc. All Rights Reserved.
+
 #include "Linux/LinuxElectraDecodersPlatformResources.h"
+#include "ElectraDecodersModule.h"
+
+#include "IElectraDecoderOutputVideo.h"
+#include "ElectraDecodersUtils.h"
+#include "MediaDecoderOutput.h"
+#include "MediaVideoDecoderOutput.h"
+#include "ElectraTextureSample.h"
+
+
+#if PLATFORM_LINUX || PLATFORM_UNIX
+
+#if WITH_LIBAV
+#include "libav_Decoder_Common_Video.h"
+#else
+class ILibavDecoderDecodedImage;
+#endif
+
+
+
+struct FElectraDecodersPlatformResourcesLinux::IDecoderPlatformResource
+{
+	uint32 MaxWidth = 0;
+	uint32 MaxHeight = 0;
+	uint32 MaxOutputBuffers = 0;
+};
+
+FElectraDecodersPlatformResourcesLinux::IDecoderPlatformResource* FElectraDecodersPlatformResourcesLinux::CreatePlatformVideoResource(void* InOwnerHandle, const TMap<FString, FVariant> InOptions)
+{
+	IDecoderPlatformResource* pr = new IDecoderPlatformResource;
+	pr->MaxWidth  = Align((uint32)ElectraDecodersUtil::GetVariantValueSafeU64(InOptions, TEXT("max_width"), 1920), 16);
+	pr->MaxHeight = Align((uint32)ElectraDecodersUtil::GetVariantValueSafeU64(InOptions, TEXT("max_height"), 1080), 16);
+	pr->MaxOutputBuffers = (uint32)ElectraDecodersUtil::GetVariantValueSafeU64(InOptions, TEXT("max_output_buffers"), 5);
+	return pr;
+}
+
+void FElectraDecodersPlatformResourcesLinux::ReleasePlatformVideoResource(void* InOwnerHandle, FElectraDecodersPlatformResourcesLinux::IDecoderPlatformResource* InHandleToRelease)
+{
+	if (InHandleToRelease)
+	{
+		delete InHandleToRelease;
+	}
+}
+
+
+#if WITH_LIBAV
+namespace LibavcodecConversion
+{
+static bool ConvertDecodedImageToNV12(TArray64<uint8>& OutNV12Buffer, FIntPoint OutBufDim, int32 NumBits, ILibavDecoderDecodedImage* InImage)
+{
+	if (!InImage || !OutNV12Buffer.GetData())
+	{
+		return false;
+	}
+	const ILibavDecoderVideoCommon::FOutputInfo& ImageInfo = InImage->GetOutputInfo();
+	if (ImageInfo.NumPlanes == 2 &&
+		ImageInfo.Planes[0].Content == ILibavDecoderVideoCommon::FPlaneInfo::EContent::Luma &&
+		ImageInfo.Planes[1].Content == ILibavDecoderVideoCommon::FPlaneInfo::EContent::ChromaUV)
+	{
+		if (ImageInfo.Planes[0].BytesPerPixel == 1)
+		{
+			const int32 w = ImageInfo.Planes[0].Width;
+			const int32 h = ImageInfo.Planes[0].Height;
+			const int32 aw = ((w + 1) / 2) * 2;
+			const int32 ah = ((h + 1) / 2) * 2;
+			uint8* DstY = OutNV12Buffer.GetData();
+			uint8* DstUV = DstY + aw * ah;
+			const uint8* SrcY = (const uint8*)ImageInfo.Planes[0].Address;
+			const uint8* SrcUV = (const uint8*)ImageInfo.Planes[1].Address;
+			// To simplify the conversion we require the output buffer to have the dimension of the planes.
+			check(OutBufDim.X == aw && OutBufDim.Y == ah);
+			if (!SrcY || !SrcUV || OutBufDim.X != aw || OutBufDim.Y != ah)
+			{
+				return false;
+			}
+			if ((w & 1) == 0)
+			{
+				FMemory::Memcpy(DstY, SrcY, w*h);
+				FMemory::Memcpy(DstUV, SrcUV, w*h/2);
+			}
+			else
+			{
+				for(int32 y=0; y<h; ++y)
+				{
+					FMemory::Memcpy(DstY, SrcY, w);
+					DstY += aw;
+					SrcY += w;
+				}
+				for(int32 y=0; y<h/2; ++y)
+				{
+					FMemory::Memcpy(DstUV, SrcUV, w);
+					DstUV += aw;
+					SrcUV += w;
+				}
+			}
+		}
+		else if (ImageInfo.Planes[0].BytesPerPixel == 2)
+		{
+			const int32 w = ImageInfo.Planes[0].Width;
+			const int32 h = ImageInfo.Planes[0].Height;
+			const int32 aw = ((w + 1) / 2) * 2;
+			const int32 ah = ((h + 1) / 2) * 2;
+			uint16* DstY = (uint16*)OutNV12Buffer.GetData();
+			uint16* DstUV = DstY + aw * ah;
+			const uint16* SrcY = (const uint16*)ImageInfo.Planes[0].Address;
+			const uint16* SrcUV = (const uint16*)ImageInfo.Planes[1].Address;
+			// To simplify the conversion we require the output buffer to have the dimension of the planes.
+			check(OutBufDim.X == aw && OutBufDim.Y == ah);
+			if (!SrcY || !SrcUV || OutBufDim.X != aw || OutBufDim.Y != ah)
+			{
+				return false;
+			}
+#if 0
+			// Copy as 16 bits
+			if ((w & 1) == 0)
+			{
+				FMemory::Memcpy(DstY, SrcY, w*h*2);
+				FMemory::Memcpy(DstUV, SrcUV, w*h);
+			}
+			else
+			{
+				for(int32 y=0; y<h; ++y)
+				{
+					FMemory::Memcpy(DstY, SrcY, w*2);
+					DstY += aw;
+					SrcY += w;
+				}
+				for(int32 y=0; y<h/2; ++y)
+				{
+					FMemory::Memcpy(DstUV, SrcUV, w*2);
+					DstUV += aw;
+					SrcUV += w;
+				}
+			}
+#else
+			// Convert down to 8 bits
+			uint8* Dst8Y = OutNV12Buffer.GetData();
+			uint8* Dst8UV = Dst8Y + aw * ah;
+			for(int32 y=0; y<h; ++y)
+			{
+				for(int32 x=0; x<w; ++x)
+				{
+					Dst8Y[x] = (*SrcY++) >> 8;
+				}
+				Dst8Y += aw;
+			}
+			for(int32 y=0; y<h/2; ++y)
+			{
+				for(int32 u=0; u<w/2; ++u)
+				{
+					Dst8UV[u*2+0] = (*SrcUV++) >> 8;
+					Dst8UV[u*2+1] = (*SrcUV++) >> 8;
+				}
+				Dst8UV += aw;
+			}
+#endif
+		}
+	}
+	else if (ImageInfo.NumPlanes == 3 &&
+				ImageInfo.Planes[0].Content == ILibavDecoderVideoCommon::FPlaneInfo::EContent::Luma &&
+				ImageInfo.Planes[1].Content == ILibavDecoderVideoCommon::FPlaneInfo::EContent::ChromaU &&
+				ImageInfo.Planes[2].Content == ILibavDecoderVideoCommon::FPlaneInfo::EContent::ChromaV)
+	{
+		if (ImageInfo.Planes[0].BytesPerPixel == 1)
+		{
+			const int32 w = ImageInfo.Planes[0].Width;
+			const int32 h = ImageInfo.Planes[0].Height;
+			const int32 aw = ((w + 1) / 2) * 2;
+			const int32 ah = ((h + 1) / 2) * 2;
+			uint8* DstY = OutNV12Buffer.GetData();
+			uint8* DstUV = DstY + aw * ah;
+			const uint8* SrcY = (const uint8*)ImageInfo.Planes[0].Address;
+			const uint8* SrcU = (const uint8*)ImageInfo.Planes[1].Address;
+			const uint8* SrcV = (const uint8*)ImageInfo.Planes[2].Address;
+			// To simplify the conversion we require the output buffer to have the dimension of the planes.
+			check(OutBufDim.X == aw && OutBufDim.Y == ah);
+			if (!SrcY || !SrcU || !SrcV || OutBufDim.X != aw || OutBufDim.Y != ah)
+			{
+				return false;
+			}
+			if ((w & 1) == 0)
+			{
+				FMemory::Memcpy(DstY, SrcY, w*h);
+				for(int32 i=0, iMax=w*h/4; i<iMax; ++i)
+				{
+					*DstUV++ = *SrcU++;
+					*DstUV++ = *SrcV++;
+				}
+			}
+			else
+			{
+				for(int32 y=0; y<h; ++y)
+				{
+					FMemory::Memcpy(DstY, SrcY, w);
+					DstY += aw;
+					SrcY += w;
+				}
+				int32 padUV = (aw - w) * 2;
+				for(int32 v=0; v<h/2; ++v)
+				{
+					for(int32 u=0; u<w/2; ++u)
+					{
+						*DstUV++ = *SrcU++;
+						*DstUV++ = *SrcV++;
+					}
+					DstUV += padUV;
+				}
+			}
+		}
+		else if (ImageInfo.Planes[0].BytesPerPixel == 2)
+		{
+			const int32 w = ImageInfo.Planes[0].Width;
+			const int32 h = ImageInfo.Planes[0].Height;
+			const int32 aw = ((w + 1) / 2) * 2;
+			const int32 ah = ((h + 1) / 2) * 2;
+		// The destination is 8 bits only!
+			uint8* DstY = OutNV12Buffer.GetData();
+			uint8* DstUV = DstY + aw * ah;
+			const uint16* SrcY = (const uint16*)ImageInfo.Planes[0].Address;
+			const uint16* SrcU = (const uint16*)ImageInfo.Planes[1].Address;
+			const uint16* SrcV = (const uint16*)ImageInfo.Planes[2].Address;
+			// To simplify the conversion we require the output buffer to have the dimension of the planes.
+			check(OutBufDim.X == aw && OutBufDim.Y == ah);
+			if (!SrcY || !SrcU || !SrcV || OutBufDim.X != aw || OutBufDim.Y != ah)
+			{
+				return false;
+			}
+			for(int32 y=0; y<h; ++y)
+			{
+				for(int32 x=0; x<w; ++x)
+				{
+					DstY[x] = SrcY[x] >> 2;
+				}
+				DstY += aw;
+				SrcY += w;
+			}
+			int32 padUV = (aw - w) * 2;
+			for(int32 v=0; v<h/2; ++v)
+			{
+				for(int32 u=0; u<w/2; ++u)
+				{
+					*DstUV++ = (*SrcU++) >> 2;
+					*DstUV++ = (*SrcV++) >> 2;
+				}
+				DstUV += padUV;
+			}
+		}
+	}
+	return true;
+}
+} // LibavcodecConversion
+#endif
+
+
 
 
-/*
-bool FElectraDecodersPlatformResourcesLinux::NotUsedRightNow()
+bool FElectraDecodersPlatformResourcesLinux::SetupOutputTextureSample(FString& OutErrorMessage, TSharedPtr<FElectraTextureSample, ESPMode::ThreadSafe> InOutTextureSampleToSetup, TSharedPtr<IElectraDecoderVideoOutput, ESPMode::ThreadSafe> InDecoderOutput, TSharedPtr<Electra::FParamDict, ESPMode::ThreadSafe> InOutBufferPropertes, FElectraDecodersPlatformResourcesLinux::IDecoderPlatformResource* InPlatformResource)
 {
+	if (!InOutTextureSampleToSetup.IsValid() || !InDecoderOutput.IsValid() || !InOutBufferPropertes.IsValid() || !InPlatformResource)
+	{
+		OutErrorMessage = TEXT("Bad parameter");
+		return false;
+	}
+
+	TMap<FString, FVariant> ExtraValues;
+	InDecoderOutput->GetExtraValues(ExtraValues);
+
+	FElectraVideoDecoderOutputCropValues Crop = InDecoderOutput->GetCropValues();
+	InOutBufferPropertes->Set(IDecoderOutputOptionNames::Width, Electra::FVariantValue((int64)InDecoderOutput->GetWidth()));
+	InOutBufferPropertes->Set(IDecoderOutputOptionNames::Height, Electra::FVariantValue((int64)InDecoderOutput->GetHeight()));
+	InOutBufferPropertes->Set(IDecoderOutputOptionNames::CropLeft, Electra::FVariantValue((int64)Crop.Left));
+	InOutBufferPropertes->Set(IDecoderOutputOptionNames::CropRight, Electra::FVariantValue((int64)Crop.Right));
+	InOutBufferPropertes->Set(IDecoderOutputOptionNames::CropTop, Electra::FVariantValue((int64)Crop.Top));
+	InOutBufferPropertes->Set(IDecoderOutputOptionNames::CropBottom, Electra::FVariantValue((int64)Crop.Bottom));
+	InOutBufferPropertes->Set(IDecoderOutputOptionNames::AspectRatio, Electra::FVariantValue((double)InDecoderOutput->GetAspectRatioW() / (double)InDecoderOutput->GetAspectRatioH()));
+	InOutBufferPropertes->Set(IDecoderOutputOptionNames::AspectW, Electra::FVariantValue((int64)InDecoderOutput->GetAspectRatioW()));
+	InOutBufferPropertes->Set(IDecoderOutputOptionNames::AspectH, Electra::FVariantValue((int64)InDecoderOutput->GetAspectRatioH()));
+	InOutBufferPropertes->Set(IDecoderOutputOptionNames::FPSNumerator, Electra::FVariantValue((int64)InDecoderOutput->GetFrameRateNumerator()));
+	InOutBufferPropertes->Set(IDecoderOutputOptionNames::FPSDenominator, Electra::FVariantValue((int64)InDecoderOutput->GetFrameRateDenominator()));
+	InOutBufferPropertes->Set(IDecoderOutputOptionNames::PixelDataScale, Electra::FVariantValue((double)ElectraDecodersUtil::GetVariantValueSafeDouble(ExtraValues, TEXT("pix_datascale"), 1.0)));
+
+	ILibavDecoderDecodedImage* DecodedImage = reinterpret_cast<ILibavDecoderDecodedImage*>(InDecoderOutput->GetPlatformOutputHandle(EElectraDecoderPlatformOutputHandleType::LibavDecoderDecodedImage));
+	IElectraDecoderVideoOutputImageBuffers* ImageBuffers = !DecodedImage ? reinterpret_cast<IElectraDecoderVideoOutputImageBuffers*>(InDecoderOutput->GetPlatformOutputHandle(EElectraDecoderPlatformOutputHandleType::ImageBuffers)) : nullptr;
+	if (DecodedImage || ImageBuffers)
+	{
+		if (DecodedImage)
+		{
+#if WITH_LIBAV
+			// We are currently converting output to 8 bit NV12 format!
+			int32 NumBits = 8;	// InDecoderOutput->GetNumberOfBits();
+			InOutBufferPropertes->Set(IDecoderOutputOptionNames::BitsPerComponent, Electra::FVariantValue((int64) NumBits));
+			InOutBufferPropertes->Set(IDecoderOutputOptionNames::PixelFormat, Electra::FVariantValue((int64)((NumBits <= 8) ? EPixelFormat::PF_NV12 : EPixelFormat::PF_P010)));
+			InOutBufferPropertes->Set(IDecoderOutputOptionNames::Pitch, Electra::FVariantValue((int64)InDecoderOutput->GetDecodedWidth() * ((NumBits <= 8) ? 1 : 2)));
+
+			const ILibavDecoderVideoCommon::FOutputInfo& DecodedImageInfo = DecodedImage->GetOutputInfo();
+			FIntPoint ImageDim(DecodedImageInfo.Planes[0].Width, DecodedImageInfo.Planes[0].Height);
+			FIntPoint AlignedDim(ImageDim);
+
+			InOutTextureSampleToSetup->InitializeCommon(*InOutBufferPropertes);
+
+			// Round up to multiple of 2.
+			AlignedDim.X = ((AlignedDim.X + 1) / 2) * 2;
+			AlignedDim.Y = ((AlignedDim.Y + 1) / 2) * 2;
+			// The vertical sample dimension encompasses the height of the UV plane.
+			AlignedDim.Y = AlignedDim.Y * 3 / 2;
+
+			InOutTextureSampleToSetup->SetDim(AlignedDim);
+			// Stride is the width of the buffer for now (since we converted down to 8 bits)
+			InOutTextureSampleToSetup->Stride = AlignedDim.X;
+			
+			// Allocate the buffer
+			InOutTextureSampleToSetup->Buffer = MakeShared<TArray64<uint8>, ESPMode::ThreadSafe>();
+			int32 AllocSize = AlignedDim.X * AlignedDim.Y * (NumBits > 8 ? 2 : 1);
+			InOutTextureSampleToSetup->Buffer->SetNum(AllocSize);
+
+			LibavcodecConversion::ConvertDecodedImageToNV12(*InOutTextureSampleToSetup->Buffer, ImageDim, NumBits, DecodedImage);
+
+			// Finish initialization.
+			if (!InOutTextureSampleToSetup->FinishInitialization())
+			{
+				OutErrorMessage = FString::Printf(TEXT("SetupOutputTextureSampleFromBuffer(): Unsupported pixel format encoding"));
+				return false;
+			}
+			return true;
+#endif
+		}
+		else if (ImageBuffers && ImageBuffers->GetNumberOfBuffers() == 1)
+		{
+			// Compatible format?
+			if ((ImageBuffers->GetBufferFormatByIndex(0) == EPixelFormat::PF_NV12 || ImageBuffers->GetBufferFormatByIndex(0) == EPixelFormat::PF_P010) &&
+				 ImageBuffers->GetBufferEncodingByIndex(0) == EElectraTextureSamplePixelEncoding::Native)
+			{
+				// We are currently converting output to 8 bit NV12 format!
+				int32 NumBits = InDecoderOutput->GetNumberOfBits();
+				int32 Pitch = ImageBuffers->GetBufferPitchByIndex(0);
+				InOutBufferPropertes->Set(IDecoderOutputOptionNames::BitsPerComponent, Electra::FVariantValue((int64) NumBits));
+				InOutBufferPropertes->Set(IDecoderOutputOptionNames::PixelFormat, Electra::FVariantValue((int64)((NumBits <= 8) ? EPixelFormat::PF_NV12 : EPixelFormat::PF_P010)));
+				InOutBufferPropertes->Set(IDecoderOutputOptionNames::Pitch, Electra::FVariantValue((int64)Pitch));
+				InOutBufferPropertes->Set(IDecoderOutputOptionNames::PixelEncoding, Electra::FVariantValue((int64)EElectraTextureSamplePixelEncoding::Native));
+
+				FIntPoint AlignedDim((int32) InDecoderOutput->GetWidth(), (int32) InDecoderOutput->GetHeight());
+
+				InOutTextureSampleToSetup->InitializeCommon(*InOutBufferPropertes);
+
+				// Round up to multiple of 2.
+				AlignedDim.X = ((AlignedDim.X + 1) / 2) * 2;
+				AlignedDim.Y = ((AlignedDim.Y + 1) / 2) * 2;
+				// The vertical sample dimension encompasses the height of the UV plane.
+				AlignedDim.Y = AlignedDim.Y * 3 / 2;
+
+				InOutTextureSampleToSetup->SetDim(AlignedDim);
+				
+				// Share the buffer
+				InOutTextureSampleToSetup->Buffer = ImageBuffers->GetBufferDataByIndex(0);
+
+				// Stride is the width of the buffer
+				InOutTextureSampleToSetup->Stride = Pitch;
+
+				// Finish initialization.
+				if (!InOutTextureSampleToSetup->FinishInitialization())
+				{
+					OutErrorMessage = FString::Printf(TEXT("SetupOutputTextureSampleFromBuffer(): Unsupported pixel format encoding"));
+					return false;
+				}
+				return true;
+			}
+		}
+	}
+	OutErrorMessage = TEXT("Unhandled decoder output format");
 	return false;
 }
-*/
\ No newline at end of file
+
+#endif // PLATFORM_LINUX || PLATFORM_UNIX
diff --git a/Engine/Plugins/Media/ElectraCodecs/Source/ElectraDecoders/Private/Utils/ElectraBitstreamProcessorGenericVideo.cpp b/Engine/Plugins/Media/ElectraCodecs/Source/ElectraDecoders/Private/Utils/ElectraBitstreamProcessorGenericVideo.cpp
--- a/Engine/Plugins/Media/ElectraCodecs/Source/ElectraDecoders/Private/Utils/ElectraBitstreamProcessorGenericVideo.cpp
+++ b/Engine/Plugins/Media/ElectraCodecs/Source/ElectraDecoders/Private/Utils/ElectraBitstreamProcessorGenericVideo.cpp
@@ -2,7 +2,6 @@
 
 #include "Utils/ElectraBitstreamProcessorGenericVideo.h"
 #include "Utils/MPEG/ElectraUtilsMPEGVideo.h"
-#include "Utils/MPEG/ElectraUtilsMP4.h"
 #include "IElectraDecoderFeaturesAndOptions.h"
 #include "ElectraDecodersUtils.h"
 
diff --git a/Engine/Plugins/Media/ElectraCodecs/Source/ElectraDecoders/Private/Utils/Google/ElectraBitstreamProcessorVPx.cpp b/Engine/Plugins/Media/ElectraCodecs/Source/ElectraDecoders/Private/Utils/Google/ElectraBitstreamProcessorVPx.cpp
--- a/Engine/Plugins/Media/ElectraCodecs/Source/ElectraDecoders/Private/Utils/Google/ElectraBitstreamProcessorVPx.cpp
+++ b/Engine/Plugins/Media/ElectraCodecs/Source/ElectraDecoders/Private/Utils/Google/ElectraBitstreamProcessorVPx.cpp
@@ -3,7 +3,7 @@
 #include "Utils/Google/ElectraBitstreamProcessorVPx.h"
 #include "Utils/Google/ElectraUtilsVPxVideo.h"
 #include "Utils/MPEG/ElectraUtilsMPEGVideo.h"	// because information may be carried in mp4 specific boxes
-#include "Utils/MPEG/ElectraUtilsMP4.h"
+#include "Utilities/UtilitiesMP4.h"
 #include "IElectraDecoderFeaturesAndOptions.h"
 #include "ElectraDecodersUtils.h"
 
diff --git a/Engine/Plugins/Media/ElectraCodecs/Source/ElectraDecoders/Private/Utils/Google/ElectraUtilsVPxVideo.cpp b/Engine/Plugins/Media/ElectraCodecs/Source/ElectraDecoders/Private/Utils/Google/ElectraUtilsVPxVideo.cpp
--- a/Engine/Plugins/Media/ElectraCodecs/Source/ElectraDecoders/Private/Utils/Google/ElectraUtilsVPxVideo.cpp
+++ b/Engine/Plugins/Media/ElectraCodecs/Source/ElectraDecoders/Private/Utils/Google/ElectraUtilsVPxVideo.cpp
@@ -1,7 +1,7 @@
 // Copyright Epic Games, Inc. All Rights Reserved.
 
 #include "Utils/Google/ElectraUtilsVPxVideo.h"
-#include "Utils/ElectraBitstreamReader.h"
+#include "Utilities/ElectraBitstream.h"
 
 
 namespace ElectraDecodersUtil
@@ -51,7 +51,7 @@ namespace ElectraDecodersUtil
 
 		bool ParseVP9UncompressedHeader(FVP9UncompressedHeader& OutHeader, const void* Data, int64 Size)
 		{
-			auto ReadColorConfig = [](FVP9UncompressedHeader& OutHeader, FElectraBitstreamReader& BitReader) -> bool
+			auto ReadColorConfig = [](FVP9UncompressedHeader& OutHeader, Electra::FBitstreamReader& BitReader) -> bool
 			{
 				if (OutHeader.GetProfile() >= 2)
 				{
@@ -90,14 +90,14 @@ namespace ElectraDecodersUtil
 				return true;
 			};
 
-			auto ReadFrameSize = [](FVP9UncompressedHeader& OutHeader, FElectraBitstreamReader& BitReader) -> bool
+			auto ReadFrameSize = [](FVP9UncompressedHeader& OutHeader, Electra::FBitstreamReader& BitReader) -> bool
 			{
 				OutHeader.frame_width = BitReader.GetBits(16) + 1;
 				OutHeader.frame_height = BitReader.GetBits(16) + 1;
 				return true;
 			};
 
-			auto ReadRenderSize = [](FVP9UncompressedHeader& OutHeader, FElectraBitstreamReader& BitReader) -> bool
+			auto ReadRenderSize = [](FVP9UncompressedHeader& OutHeader, Electra::FBitstreamReader& BitReader) -> bool
 			{
 				if (BitReader.GetBits(1) /*render_and_frame_size_different*/)
 				{
@@ -116,7 +116,7 @@ namespace ElectraDecodersUtil
 			{
 				return false;
 			}
-			FElectraBitstreamReader BitReader(Data, (uint64) Size);
+			Electra::FBitstreamReader BitReader(Data, (uint64) Size);
 			if ((OutHeader.frame_marker = BitReader.GetBits(2)) != 2)
 			{
 				return false;
diff --git a/Engine/Plugins/Media/ElectraCodecs/Source/ElectraDecoders/Private/Utils/MPEG/ElectraBitstreamProcessor_H264.cpp b/Engine/Plugins/Media/ElectraCodecs/Source/ElectraDecoders/Private/Utils/MPEG/ElectraBitstreamProcessor_H264.cpp
--- a/Engine/Plugins/Media/ElectraCodecs/Source/ElectraDecoders/Private/Utils/MPEG/ElectraBitstreamProcessor_H264.cpp
+++ b/Engine/Plugins/Media/ElectraCodecs/Source/ElectraDecoders/Private/Utils/MPEG/ElectraBitstreamProcessor_H264.cpp
@@ -3,7 +3,7 @@
 #include "Utils/MPEG/ElectraBitstreamProcessor_H264.h"
 #include "Utils/MPEG/ElectraUtilsMPEGVideo_H264.h"
 #include "Utils/MPEG/ElectraUtilsMPEGVideo.h"
-#include "Utils/MPEG/ElectraUtilsMP4.h"
+#include "Utilities/UtilitiesMP4.h"
 #include "IElectraDecoderFeaturesAndOptions.h"
 #include "ElectraDecodersUtils.h"
 
@@ -275,11 +275,16 @@ IElectraDecoderBitstreamProcessor::EProcessResult FElectraDecoderBitstreamProces
 	uint32* LastNALU = (uint32*)ElectraDecodersUtil::AdvancePointer(NALU, InOutAccessUnit.DataSize);
 	while(NALU < LastNALU)
 	{
-		uint32 naluLen = ElectraDecodersUtil::GetFromBigEndian(*NALU);
+		uint32 naluLen = Electra::GetFromBigEndian(*NALU);
 
 		// Check the nal_ref_idc in the NAL unit for dependencies.
 		uint8 nal = *(const uint8*)(NALU + 1);
-		check((nal & 0x80) == 0);
+		if ((nal & 0x80) != 0)
+		{
+			Impl->LastErrorMessage = FString::Printf(TEXT("NAL zero bit not zero"));
+			return EProcessResult::Error;
+		}
+
 		if ((nal >> 5) != 0)
 		{
 			InOutAccessUnit.Flags &= ~EElectraDecoderFlags::IsDiscardable;
@@ -310,7 +315,7 @@ IElectraDecoderBitstreamProcessor::EProcessResult FElectraDecoderBitstreamProces
 
 		if (Impl->bReplaceLengthWithStartcode)
 		{
-			*NALU = ElectraDecodersUtil::GetFromBigEndian(0x00000001U);
+			*NALU = Electra::GetFromBigEndian(0x00000001U);
 		}
 		NALU = ElectraDecodersUtil::AdvancePointer(NALU, naluLen + 4);
 	}
@@ -434,7 +439,7 @@ bool FElectraDecoderBitstreamProcessorH264::FImpl::HandlePicTiming(const Electra
 	auto It = InSPSs.CreateConstIterator();
 	const ElectraDecodersUtil::MPEG::H264::FSequenceParameterSet& sps(It->Value);
 
-	ElectraDecodersUtil::MPEG::H264::FBitstreamReader br(InSEI.Message.GetData(), InSEI.Message.Num());
+	ElectraDecodersUtil::MPEG::H264::FBitstreamReaderH264 br(InSEI.Message.GetData(), InSEI.Message.Num());
 
 	const bool CpbDpbDelaysPresentFlag = sps.nal_hrd_parameters_present_flag || sps.vcl_hrd_parameters_present_flag;
 	if (CpbDpbDelaysPresentFlag)
diff --git a/Engine/Plugins/Media/ElectraCodecs/Source/ElectraDecoders/Private/Utils/MPEG/ElectraBitstreamProcessor_H265.cpp b/Engine/Plugins/Media/ElectraCodecs/Source/ElectraDecoders/Private/Utils/MPEG/ElectraBitstreamProcessor_H265.cpp
--- a/Engine/Plugins/Media/ElectraCodecs/Source/ElectraDecoders/Private/Utils/MPEG/ElectraBitstreamProcessor_H265.cpp
+++ b/Engine/Plugins/Media/ElectraCodecs/Source/ElectraDecoders/Private/Utils/MPEG/ElectraBitstreamProcessor_H265.cpp
@@ -3,7 +3,7 @@
 #include "Utils/MPEG/ElectraBitstreamProcessor_H265.h"
 #include "Utils/MPEG/ElectraUtilsMPEGVideo_H265.h"
 #include "Utils/MPEG/ElectraUtilsMPEGVideo.h"
-#include "Utils/MPEG/ElectraUtilsMP4.h"
+#include "Utilities/UtilitiesMP4.h"
 #include "IElectraDecoderFeaturesAndOptions.h"
 #include "ElectraDecodersUtils.h"
 
@@ -278,11 +278,15 @@ IElectraDecoderBitstreamProcessor::EProcessResult FElectraDecoderBitstreamProces
 	uint32* LastNALU = (uint32*)ElectraDecodersUtil::AdvancePointer(NALU, InOutAccessUnit.DataSize);
 	while(NALU < LastNALU)
 	{
-		uint32 naluLen = ElectraDecodersUtil::GetFromBigEndian(*NALU);
+		uint32 naluLen = Electra::GetFromBigEndian(*NALU);
 
 		// Check the nal_ref_idc in the NAL unit for dependencies.
 		uint8 nut = *(const uint8*)(NALU + 1);
-		check((nut & 0x80) == 0);
+		if ((nut & 0x80) != 0)
+		{
+			Impl->LastErrorMessage = FString::Printf(TEXT("NUT zero bit not zero"));
+			return EProcessResult::Error;
+		}
 		nut >>= 1;
 
 		// IDR, CRA or BLA frame?
@@ -305,7 +309,7 @@ IElectraDecoderBitstreamProcessor::EProcessResult FElectraDecoderBitstreamProces
 
 		if (Impl->bReplaceLengthWithStartcode)
 		{
-			*NALU = ElectraDecodersUtil::GetFromBigEndian(0x00000001U);
+			*NALU = Electra::GetFromBigEndian(0x00000001U);
 		}
 		NALU = ElectraDecodersUtil::AdvancePointer(NALU, naluLen + 4);
 	}
@@ -482,7 +486,7 @@ bool FElectraDecoderBitstreamProcessorH265::FImpl::HandleTimeCode(const ElectraD
 	}
 	const ElectraDecodersUtil::MPEG::H265::FVideoParameterSet& vps(VPSs[sps.sps_video_parameter_set_id]);
 
-	ElectraDecodersUtil::MPEG::H265::FBitstreamReader br(InSEI.Message.GetData(), InSEI.Message.Num());
+	ElectraDecodersUtil::MPEG::H265::FBitstreamReaderH265 br(InSEI.Message.GetData(), InSEI.Message.Num());
 
 	uint32 num_clock_ts = br.GetBits(2);
 	for(uint32 nc=0; nc<num_clock_ts; ++nc)
