ElectraCodecs: Removed DX12 resource transitions that are not needed on the copy queue and trigger on the latest d3d debug validation layer.

#rnx

diff --git a/Engine/Plugins/Media/ElectraCodecs/Source/ElectraDecoders/Private/Windows/WindowsElectraDecodersPlatformResources.cpp b/Engine/Plugins/Media/ElectraCodecs/Source/ElectraDecoders/Private/Windows/WindowsElectraDecodersPlatformResources.cpp
--- a/Engine/Plugins/Media/ElectraCodecs/Source/ElectraDecoders/Private/Windows/WindowsElectraDecodersPlatformResources.cpp
+++ b/Engine/Plugins/Media/ElectraCodecs/Source/ElectraDecoders/Private/Windows/WindowsElectraDecodersPlatformResources.cpp
@@ -357,21 +357,9 @@ static bool SetupOutputTextureSampleFromDX12Resource(FString& OutErrorMessage, T
 	// - - - - - - - - - - - - - - - - - - - - - - - - - - - - - - - - - - - - - - - - - - - - - - - - - - - - - - - -
 	// Build command list to copy data from decoder buffer
 
-	D3D12_RESOURCE_BARRIER PreCopyResourceBarriers[] = {
-		{D3D12_RESOURCE_BARRIER_TYPE_TRANSITION, D3D12_RESOURCE_BARRIER_FLAG_NONE, {{ InResource, D3D12_RESOURCE_BARRIER_ALL_SUBRESOURCES, D3D12_RESOURCE_STATE_COMMON, D3D12_RESOURCE_STATE_COPY_SOURCE }}},
-		{D3D12_RESOURCE_BARRIER_TYPE_TRANSITION, D3D12_RESOURCE_BARRIER_FLAG_NONE, {{ InOutTextureSampleToSetup->TextureDX12, D3D12_RESOURCE_BARRIER_ALL_SUBRESOURCES, D3D12_RESOURCE_STATE_COMMON, D3D12_RESOURCE_STATE_COPY_DEST}}},
-	};
-
-	D3D12_RESOURCE_BARRIER PostCopyResourceBarriers[] = {
-		{ D3D12_RESOURCE_BARRIER_TYPE_TRANSITION, D3D12_RESOURCE_BARRIER_FLAG_NONE, {{ InResource, D3D12_RESOURCE_BARRIER_ALL_SUBRESOURCES, D3D12_RESOURCE_STATE_COPY_SOURCE, D3D12_RESOURCE_STATE_COMMON }}},
-		{ D3D12_RESOURCE_BARRIER_TYPE_TRANSITION, D3D12_RESOURCE_BARRIER_FLAG_NONE, {{ InOutTextureSampleToSetup->TextureDX12, D3D12_RESOURCE_BARRIER_ALL_SUBRESOURCES, D3D12_RESOURCE_STATE_COPY_DEST, D3D12_RESOURCE_STATE_COMMON }}},
-	};
-
 	InOutTextureSampleToSetup->D3DCmdAllocator->Reset();
 	InOutTextureSampleToSetup->D3DCmdList->Reset(InOutTextureSampleToSetup->D3DCmdAllocator, nullptr);
 
-	InOutTextureSampleToSetup->D3DCmdList->ResourceBarrier(2, PreCopyResourceBarriers);
-
 	if (bIsTexture)
 	{
 		// If we got a texture, we simply can copy the resource as a whole...
@@ -443,8 +431,6 @@ static bool SetupOutputTextureSampleFromDX12Resource(FString& OutErrorMessage, T
 		}
 	}
 
-	InOutTextureSampleToSetup->D3DCmdList->ResourceBarrier(2, PostCopyResourceBarriers);
-
 	InOutTextureSampleToSetup->D3DCmdList->Close();
 
 	// - - - - - - - - - - - - - - - - - - - - - - - - - - - - - - - - - - - - - - - - - - - - - - - - - - - - - - - -
