Verify D3D12 results for Electra CopyBufferData

diff --git a/Engine/Plugins/Media/ElectraCodecs/Source/ElectraDecoders/Private/DX/VideoDecoderH264_DX.cpp b/Engine/Plugins/Media/ElectraCodecs/Source/ElectraDecoders/Private/DX/VideoDecoderH264_DX.cpp
--- a/Engine/Plugins/Media/ElectraCodecs/Source/ElectraDecoders/Private/DX/VideoDecoderH264_DX.cpp
+++ b/Engine/Plugins/Media/ElectraCodecs/Source/ElectraDecoders/Private/DX/VideoDecoderH264_DX.cpp
@@ -21,6 +21,7 @@
 
 #include "Stats/Stats.h"
 #include "HAL/IConsoleManager.h"
+#include "ID3D12DynamicRHI.h"
 
 /*********************************************************************************************************************/
 
@@ -974,9 +975,15 @@ void FElectraVideoDecoderH264_DX::CopyBufferData(TRefCountPtr<IMF2DBuffer> Buffe
 	HRESULT Res = Buffer2D->Lock2D(&Data, &Pitch);
 	check(SUCCEEDED(Res));
 
+	ID3D12DynamicRHI* const D3D12RHI = GetID3D12DynamicRHI();
+	ID3D12Device* const Device = static_cast<ID3D12Device*>(D3D12RHI->RHIGetNativeDevice());
+
 	void* BufferAddr;
 	Res = BufferResource->Map(0, nullptr, &BufferAddr);
-	check(SUCCEEDED(Res));
+	if (FAILED(Res))
+	{
+		D3D12RHI->RHIVerifyResult(Device, Res, "FElectraVideoDecoderH264_DX::CopyBufferData Map", __FILE__, __LINE__);
+	}
 
 	FElectraMediaDecoderOutputBufferPool_DX12::CopyWithPitchAdjust((uint8*)BufferAddr, BufferPitch, Data, Pitch, DecodedHeight * 3 / 2);
 
@@ -984,7 +991,11 @@ void FElectraVideoDecoderH264_DX::CopyBufferData(TRefCountPtr<IMF2DBuffer> Buffe
 
 	Buffer2D->Unlock2D();
 
-	BufferFence->Signal(BufferFenceValue);
+	Res = BufferFence->Signal(BufferFenceValue);
+	if (FAILED(Res))
+	{
+		D3D12RHI->RHIVerifyResult(Device, Res, "FElectraVideoDecoderH264_DX::CopyBufferData Signal", __FILE__, __LINE__);
+	}
 };
 #endif
 
diff --git a/Engine/Plugins/Media/ElectraCodecs/Source/ElectraDecoders/Private/DX/VideoDecoderH265_DX.cpp b/Engine/Plugins/Media/ElectraCodecs/Source/ElectraDecoders/Private/DX/VideoDecoderH265_DX.cpp
--- a/Engine/Plugins/Media/ElectraCodecs/Source/ElectraDecoders/Private/DX/VideoDecoderH265_DX.cpp
+++ b/Engine/Plugins/Media/ElectraCodecs/Source/ElectraDecoders/Private/DX/VideoDecoderH265_DX.cpp
@@ -20,6 +20,7 @@
 
 #include "Stats/Stats.h"
 #include "HAL/IConsoleManager.h"
+#include "ID3D12DynamicRHI.h"
 
 /*********************************************************************************************************************/
 
@@ -959,9 +960,15 @@ void FElectraVideoDecoderH265_DX::CopyBufferData(TRefCountPtr<IMF2DBuffer> Buffe
 	HRESULT Res = Buffer2D->Lock2D(&Data, &Pitch);
 	check(SUCCEEDED(Res));
 
+	ID3D12DynamicRHI* const D3D12RHI = GetID3D12DynamicRHI();
+	ID3D12Device* const Device = static_cast<ID3D12Device*>(D3D12RHI->RHIGetNativeDevice());
+
 	void* BufferAddr;
 	Res = BufferResource->Map(0, nullptr, &BufferAddr);
-	check(SUCCEEDED(Res));
+	if (FAILED(Res))
+	{
+		D3D12RHI->RHIVerifyResult(Device, Res, "FElectraVideoDecoderH265_DX::CopyBufferData Map", __FILE__, __LINE__);
+	}
 
 	FElectraMediaDecoderOutputBufferPool_DX12::CopyWithPitchAdjust((uint8*)BufferAddr, BufferPitch, Data, Pitch, DecodedHeight * 3 / 2);
 
@@ -969,7 +976,11 @@ void FElectraVideoDecoderH265_DX::CopyBufferData(TRefCountPtr<IMF2DBuffer> Buffe
 
 	Buffer2D->Unlock2D();
 
-	BufferFence->Signal(BufferFenceValue);
+	Res = BufferFence->Signal(BufferFenceValue);
+	if (FAILED(Res))
+	{
+		D3D12RHI->RHIVerifyResult(Device, Res, "FElectraVideoDecoderH265_DX::CopyBufferData Signal", __FILE__, __LINE__);
+	}
 };
 #endif
 
