[UBT][CppCompileWarnings] Engine compliance wundef 5.7

diff --git a/Engine/Plugins/FX/Niagara/Source/Niagara/Public/NiagaraScalabilityState.h b/Engine/Plugins/FX/Niagara/Source/Niagara/Public/NiagaraScalabilityState.h
--- a/Engine/Plugins/FX/Niagara/Source/Niagara/Public/NiagaraScalabilityState.h
+++ b/Engine/Plugins/FX/Niagara/Source/Niagara/Public/NiagaraScalabilityState.h
@@ -2,6 +2,8 @@
 
 #pragma once
 
+#include "NiagaraDefines.h"
+
 #include "NiagaraScalabilityState.generated.h"
 
 USTRUCT()
diff --git a/Engine/Plugins/Media/ElectraCodecs/Source/ElectraDecoders/Private/Windows/WindowsElectraDecodersPlatformResources.cpp b/Engine/Plugins/Media/ElectraCodecs/Source/ElectraDecoders/Private/Windows/WindowsElectraDecodersPlatformResources.cpp
--- a/Engine/Plugins/Media/ElectraCodecs/Source/ElectraDecoders/Private/Windows/WindowsElectraDecodersPlatformResources.cpp
+++ b/Engine/Plugins/Media/ElectraCodecs/Source/ElectraDecoders/Private/Windows/WindowsElectraDecodersPlatformResources.cpp
@@ -2,7 +2,7 @@
 
 #include "WindowsElectraDecodersPlatformResources.h"
 #include "ElectraDecodersModule.h"
-
+#include "WindowsPlatformHeaders_Video_DX.h"
 #include "ID3D12DynamicRHI.h"
 #include "Async/Async.h"
 
@@ -115,13 +115,13 @@ static void TriggerDataCopy(TRefCountPtr<ID3D12GraphicsCommandList> D3DCmdList,
 		[CmdList = D3DCmdList, CmdAllocator = D3DCmdAllocator, DestFence = D3DFence, DestFenceValue = FenceValue, OutputSync](ID3D12CommandQueue* D3DCmdQueue)
 		{
 			TRefCountPtr<ID3D12Fence> ResourceFence;
-	#if ALLOW_MFSAMPLE_WITH_DX12
+	#if ELECTRA_DECODERS_ENABLE_DX && ALLOW_MFSAMPLE_WITH_DX12
 			TRefCountPtr<IMFD3D12SynchronizationObjectCommands> DecoderSync;
 	#endif
 
 			if (OutputSync.Sync.IsValid())
 			{
-	#if ALLOW_MFSAMPLE_WITH_DX12
+	#if ELECTRA_DECODERS_ENABLE_DX && ALLOW_MFSAMPLE_WITH_DX12
 				if (OutputSync.Sync->QueryInterface(__uuidof(IMFD3D12SynchronizationObjectCommands), (void**)DecoderSync.GetInitReference()) != S_OK)
 	#endif
 				{
@@ -129,7 +129,7 @@ static void TriggerDataCopy(TRefCountPtr<ID3D12GraphicsCommandList> D3DCmdList,
 				}
 			}
 
-	#if ALLOW_MFSAMPLE_WITH_DX12
+	#if ELECTRA_DECODERS_ENABLE_DX && ALLOW_MFSAMPLE_WITH_DX12
 			if (DecoderSync)
 			{
 				// Sync queue to make sure decoder output is ready for us (WMF case)
@@ -146,7 +146,7 @@ static void TriggerDataCopy(TRefCountPtr<ID3D12GraphicsCommandList> D3DCmdList,
 			ID3D12CommandList* CmdLists[1] = { CmdList.GetReference() };
 			D3DCmdQueue->ExecuteCommandLists(1, CmdLists);
 
-	#if ALLOW_MFSAMPLE_WITH_DX12
+	#if ELECTRA_DECODERS_ENABLE_DX && ALLOW_MFSAMPLE_WITH_DX12
 			if (DecoderSync)
 			{
 				// Sync to end of copy operation and release MFSample once reached
diff --git a/Engine/Source/Programs/UnrealFrontend/Private/Commands/UserInterfaceCommand.cpp b/Engine/Source/Programs/UnrealFrontend/Private/Commands/UserInterfaceCommand.cpp
--- a/Engine/Source/Programs/UnrealFrontend/Private/Commands/UserInterfaceCommand.cpp
+++ b/Engine/Source/Programs/UnrealFrontend/Private/Commands/UserInterfaceCommand.cpp
@@ -60,7 +60,7 @@ void FUserInterfaceCommand::Run(  )
 
 	// load optional modules
 	FModuleManager::Get().LoadModule("DeviceManager");
-#if STATS && UE_DEPRECATED_PROFILER_ENABLED
+#if STATS && defined(UE_DEPRECATED_PROFILER_ENABLED) && UE_DEPRECATED_PROFILER_ENABLED
 	FModuleManager::Get().LoadModule("ProfilerClient");
 #endif
 	FModuleManager::Get().LoadModule("LegacyProjectLauncher");
diff --git a/Engine/Source/Runtime/Core/Public/HAL/Memory.h b/Engine/Source/Runtime/Core/Public/HAL/Memory.h
--- a/Engine/Source/Runtime/Core/Public/HAL/Memory.h
+++ b/Engine/Source/Runtime/Core/Public/HAL/Memory.h
@@ -2,7 +2,7 @@
 
 #pragma once
 
-#if PLATFORM_USES_FIXED_GMalloc_CLASS
+#if defined(PLATFORM_USES_FIXED_GMalloc_CLASS) && PLATFORM_USES_FIXED_GMalloc_CLASS
 #	if USE_MALLOC_BINNED2
 #		include "HAL/MallocBinned2.h" // HEADER_UNIT_IGNORE - Causes circular includes
 #		define FMEMORY_INLINE_GMalloc (FMallocBinned2::MallocBinned2)
diff --git a/Engine/Source/Runtime/JsonUtilities/Private/JsonUClassExporter.cpp b/Engine/Source/Runtime/JsonUtilities/Private/JsonUClassExporter.cpp
--- a/Engine/Source/Runtime/JsonUtilities/Private/JsonUClassExporter.cpp
+++ b/Engine/Source/Runtime/JsonUtilities/Private/JsonUClassExporter.cpp
@@ -1,5 +1,7 @@
 // Copyright Epic Games, Inc. All Rights Reserved.
 
+#include "Misc/Build.h"
+
 #if !UE_BUILD_SHIPPING
 
 #include "HAL/IConsoleManager.h"
