ElectraDecoders: Only register the MF codec factories when the required DLLs could be loaded

#rnx

diff --git a/Engine/Plugins/Media/ElectraCodecs/Source/ElectraDecoders/Private/Windows/WindowsPlatformElectraDecoders.cpp b/Engine/Plugins/Media/ElectraCodecs/Source/ElectraDecoders/Private/Windows/WindowsPlatformElectraDecoders.cpp
--- a/Engine/Plugins/Media/ElectraCodecs/Source/ElectraDecoders/Private/Windows/WindowsPlatformElectraDecoders.cpp
+++ b/Engine/Plugins/Media/ElectraCodecs/Source/ElectraDecoders/Private/Windows/WindowsPlatformElectraDecoders.cpp
@@ -71,7 +71,7 @@ void FPlatformElectraDecodersWindows::Shutdown()
 void FPlatformElectraDecodersWindows::RegisterWithCodecFactory(IElectraCodecRegistry* InCodecFactoryToRegisterWith)
 {
 	check(InCodecFactoryToRegisterWith);
-	if (InCodecFactoryToRegisterWith)
+	if (InCodecFactoryToRegisterWith && PlatformElectraDecodersWindows::bMFInitialized)
 	{
 #ifdef ELECTRA_DECODERS_ENABLE_DX
 		auto RegisterFactory = [&](TSharedPtr<IElectraCodecFactory, ESPMode::ThreadSafe> InFactory) -> void
