ElectraDecoders: Added missing explicit ESPMode on shared pointers of D3D helper for consistency

#rnx

diff --git a/Engine/Plugins/Media/AppleProResDecoderElectra/Source/AppleProResDecoderElectra/Private/ProResDecoder/Windows/ElectraMediaProResDecoder.cpp b/Engine/Plugins/Media/AppleProResDecoderElectra/Source/AppleProResDecoderElectra/Private/ProResDecoder/Windows/ElectraMediaProResDecoder.cpp
--- a/Engine/Plugins/Media/AppleProResDecoderElectra/Source/AppleProResDecoderElectra/Private/ProResDecoder/Windows/ElectraMediaProResDecoder.cpp
+++ b/Engine/Plugins/Media/AppleProResDecoderElectra/Source/AppleProResDecoderElectra/Private/ProResDecoder/Windows/ElectraMediaProResDecoder.cpp
@@ -377,7 +377,7 @@ private:
 
 	uint32 MaxOutputBuffers;
 #if ELECTRA_MEDIAGPUBUFFER_DX12
-	mutable TSharedPtr<FElectraMediaDecoderOutputBufferPool_DX12> D3D12ResourcePool;
+	mutable TSharedPtr<FElectraMediaDecoderOutputBufferPool_DX12, ESPMode::ThreadSafe> D3D12ResourcePool;
 #endif
 };
 
diff --git a/Engine/Plugins/Media/AvidDNxHDDecoderElectra/Source/Private/Windows/ElectraMediaAvidDNxHDDecoder.cpp b/Engine/Plugins/Media/AvidDNxHDDecoderElectra/Source/Private/Windows/ElectraMediaAvidDNxHDDecoder.cpp
--- a/Engine/Plugins/Media/AvidDNxHDDecoderElectra/Source/Private/Windows/ElectraMediaAvidDNxHDDecoder.cpp
+++ b/Engine/Plugins/Media/AvidDNxHDDecoderElectra/Source/Private/Windows/ElectraMediaAvidDNxHDDecoder.cpp
@@ -285,7 +285,7 @@ private:
 
 	uint32 MaxOutputBuffers;
 #if ELECTRA_MEDIAGPUBUFFER_DX12
-	mutable TSharedPtr<FElectraMediaDecoderOutputBufferPool_DX12> D3D12ResourcePool;
+	mutable TSharedPtr<FElectraMediaDecoderOutputBufferPool_DX12, ESPMode::ThreadSafe> D3D12ResourcePool;
 	mutable TSharedPtr<FElectraDecodersPlatformResources::FAsyncConsecutiveTaskSync, ESPMode::ThreadSafe> TaskSync;
 #endif
 };
diff --git a/Engine/Plugins/Media/ElectraCodecs/Source/ElectraDecoders/Private/DX/VideoDecoderH264_DX.cpp b/Engine/Plugins/Media/ElectraCodecs/Source/ElectraDecoders/Private/DX/VideoDecoderH264_DX.cpp
--- a/Engine/Plugins/Media/ElectraCodecs/Source/ElectraDecoders/Private/DX/VideoDecoderH264_DX.cpp
+++ b/Engine/Plugins/Media/ElectraCodecs/Source/ElectraDecoders/Private/DX/VideoDecoderH264_DX.cpp
@@ -372,7 +372,7 @@ private:
 	uint32 MaxWidth;
 	uint32 MaxHeight;
 #if !ALLOW_MFSAMPLE_WITH_DX12
-	mutable TSharedPtr<FElectraMediaDecoderOutputBufferPool_DX12> D3D12ResourcePool;
+	mutable TSharedPtr<FElectraMediaDecoderOutputBufferPool_DX12, ESPMode::ThreadSafe> D3D12ResourcePool;
 	mutable TSharedPtr<FElectraDecodersPlatformResources::FAsyncConsecutiveTaskSync, ESPMode::ThreadSafe> TaskSync;
 #endif
 
diff --git a/Engine/Plugins/Media/ElectraCodecs/Source/ElectraDecoders/Private/DX/VideoDecoderH265_DX.cpp b/Engine/Plugins/Media/ElectraCodecs/Source/ElectraDecoders/Private/DX/VideoDecoderH265_DX.cpp
--- a/Engine/Plugins/Media/ElectraCodecs/Source/ElectraDecoders/Private/DX/VideoDecoderH265_DX.cpp
+++ b/Engine/Plugins/Media/ElectraCodecs/Source/ElectraDecoders/Private/DX/VideoDecoderH265_DX.cpp
@@ -369,7 +369,7 @@ private:
 	uint32 MaxWidth;
 	uint32 MaxHeight;
 #if !ALLOW_MFSAMPLE_WITH_DX12
-	mutable TSharedPtr<FElectraMediaDecoderOutputBufferPool_DX12> D3D12ResourcePool;
+	mutable TSharedPtr<FElectraMediaDecoderOutputBufferPool_DX12, ESPMode::ThreadSafe> D3D12ResourcePool;
 	mutable TSharedPtr<FElectraDecodersPlatformResources::FAsyncConsecutiveTaskSync, ESPMode::ThreadSafe> TaskSync;
 #endif
 
diff --git a/Engine/Plugins/Media/ElectraCodecs/Source/ElectraDecoders/Private/Windows/WindowsElectraDecodersPlatformResources.cpp b/Engine/Plugins/Media/ElectraCodecs/Source/ElectraDecoders/Private/Windows/WindowsElectraDecodersPlatformResources.cpp
--- a/Engine/Plugins/Media/ElectraCodecs/Source/ElectraDecoders/Private/Windows/WindowsElectraDecodersPlatformResources.cpp
+++ b/Engine/Plugins/Media/ElectraCodecs/Source/ElectraDecoders/Private/Windows/WindowsElectraDecodersPlatformResources.cpp
@@ -33,7 +33,7 @@ struct FElectraDecodersPlatformResourcesWindows::IDecoderPlatformResource
 	uint32 MaxHeight = 0;
 	uint32 MaxOutputBuffers = 0;
 
-	TSharedPtr<FElectraMediaDecoderOutputBufferPool_DX12> D3D12ResourcePool;
+	TSharedPtr<FElectraMediaDecoderOutputBufferPool_DX12, ESPMode::ThreadSafe> D3D12ResourcePool;
 };
 
 FElectraDecodersPlatformResourcesWindows::IDecoderPlatformResource* FElectraDecodersPlatformResourcesWindows::CreatePlatformVideoResource(void* InOwnerHandle, const TMap<FString, FVariant> InOptions)
@@ -168,7 +168,7 @@ static void TriggerDataCopy(TRefCountPtr<ID3D12GraphicsCommandList> D3DCmdList,
 static bool SetupOutputTextureSampleFromDX12Resource(FString& OutErrorMessage, TSharedPtr<FElectraTextureSample, ESPMode::ThreadSafe> InOutTextureSampleToSetup,
 													 TRefCountPtr<ID3D12Device>& InD3D12Device, TRefCountPtr<ID3D12Resource> InResource, uint32 ResourcePitch,
 													 const FElectraDecoderOutputSync& OutputSync, const FIntPoint& InSampleDim, const Electra::FParamDict& InParamDict,
-													 TSharedPtr<FElectraMediaDecoderOutputBufferPool_DX12>& InOutD3D12ResourcePool,
+													 TSharedPtr<FElectraMediaDecoderOutputBufferPool_DX12, ESPMode::ThreadSafe>& InOutD3D12ResourcePool,
 													 uint32 MaxWidth, uint32 MaxHeight, uint32 MaxOutputBuffers)
 {
 	HRESULT Result;
@@ -324,7 +324,7 @@ static bool SetupOutputTextureSampleFromDX12Resource(FString& OutErrorMessage, T
 
 		if (!InOutTextureSampleToSetup->D3D12ResourcePool.IsValid() || !InOutTextureSampleToSetup->D3D12ResourcePool->IsCompatibleAsTexture(MaxOutputBuffers + kElectraDecoderPipelineExtraFrames, MaxWidth, MaxHeight, DXGIFmt))
 		{
-			InOutTextureSampleToSetup->D3D12ResourcePool = MakeShared<FElectraMediaDecoderOutputBufferPool_DX12>(InD3D12Device, MaxOutputBuffers + kElectraDecoderPipelineExtraFrames, MaxWidth, MaxHeight, DXGIFmt, D3D12_HEAP_TYPE_DEFAULT);
+			InOutTextureSampleToSetup->D3D12ResourcePool = MakeShared<FElectraMediaDecoderOutputBufferPool_DX12, ESPMode::ThreadSafe>(InD3D12Device, MaxOutputBuffers + kElectraDecoderPipelineExtraFrames, MaxWidth, MaxHeight, DXGIFmt, D3D12_HEAP_TYPE_DEFAULT);
 			if (!InOutTextureSampleToSetup->D3D12ResourcePool.IsValid())
 			{
 				OutErrorMessage = FString::Printf(TEXT("SetupOutputTextureSampleFromDX12Resource(): Could not allocate texture resource heap"));
diff --git a/Engine/Plugins/Media/ElectraUtil/Source/ElectraSamples/Public/Windows/ElectraTextureSample.h b/Engine/Plugins/Media/ElectraUtil/Source/ElectraSamples/Public/Windows/ElectraTextureSample.h
--- a/Engine/Plugins/Media/ElectraUtil/Source/ElectraSamples/Public/Windows/ElectraTextureSample.h
+++ b/Engine/Plugins/Media/ElectraUtil/Source/ElectraSamples/Public/Windows/ElectraTextureSample.h
@@ -105,7 +105,7 @@ public:
 	TRefCountPtr<ID3D12Resource> DecoderOutputResource;
 	TRefCountPtr<ID3D12CommandAllocator> D3DCmdAllocator;
 	TRefCountPtr<ID3D12GraphicsCommandList> D3DCmdList;
-	TSharedPtr<FElectraMediaDecoderOutputBufferPool_DX12> D3D12ResourcePool;
+	TSharedPtr<FElectraMediaDecoderOutputBufferPool_DX12, ESPMode::ThreadSafe> D3D12ResourcePool;
 
 	// CPU-side buffer
 	TSharedPtr<TArray64<uint8>, ESPMode::ThreadSafe> Buffer;
diff --git a/Engine/Plugins/Media/HAPDecoderElectra/Source/HAPDecoderElectra/Private/HAPDecoder/ElectraMediaHAPDecoder.cpp b/Engine/Plugins/Media/HAPDecoderElectra/Source/HAPDecoderElectra/Private/HAPDecoder/ElectraMediaHAPDecoder.cpp
--- a/Engine/Plugins/Media/HAPDecoderElectra/Source/HAPDecoderElectra/Private/HAPDecoder/ElectraMediaHAPDecoder.cpp
+++ b/Engine/Plugins/Media/HAPDecoderElectra/Source/HAPDecoderElectra/Private/HAPDecoder/ElectraMediaHAPDecoder.cpp
@@ -315,7 +315,7 @@ private:
 		uint32 AllocSize = 0;
 		uint32 AllocBlockSize = 0;
 #if ELECTRA_MEDIAGPUBUFFER_DX12
-		mutable TSharedPtr<FElectraMediaDecoderOutputBufferPool_DX12> D3D12ResourcePool;
+		mutable TSharedPtr<FElectraMediaDecoderOutputBufferPool_DX12, ESPMode::ThreadSafe> D3D12ResourcePool;
 #endif
 	};
 	FBufferInstanceInfo InstanceBufferData[FVideoDecoderOutputHAPElectra::Buffer_Max];
diff --git a/Engine/Plugins/Media/VPxDecoderElectra/Source/VPxDecoderElectra/Private/VPxDecoder/ElectraMediaVPxDecoder.cpp b/Engine/Plugins/Media/VPxDecoderElectra/Source/VPxDecoderElectra/Private/VPxDecoder/ElectraMediaVPxDecoder.cpp
--- a/Engine/Plugins/Media/VPxDecoderElectra/Source/VPxDecoderElectra/Private/VPxDecoder/ElectraMediaVPxDecoder.cpp
+++ b/Engine/Plugins/Media/VPxDecoderElectra/Source/VPxDecoderElectra/Private/VPxDecoder/ElectraMediaVPxDecoder.cpp
@@ -311,7 +311,7 @@ private:
 	uint32 MaxHeight;
 	uint32 MaxOutputBuffers;
 #if ELECTRA_MEDIAGPUBUFFER_DX12
-	mutable TSharedPtr<FElectraMediaDecoderOutputBufferPool_DX12> D3D12ResourcePool;
+	mutable TSharedPtr<FElectraMediaDecoderOutputBufferPool_DX12, ESPMode::ThreadSafe> D3D12ResourcePool;
 #endif
 };
 
