ElectraCodecs: Added support code for AV1

#rnx

diff --git a/Engine/Plugins/Media/ElectraCodecs/Source/ElectraDecoders/Private/ElectraDecodersUtils.cpp b/Engine/Plugins/Media/ElectraCodecs/Source/ElectraDecoders/Private/ElectraDecodersUtils.cpp
--- a/Engine/Plugins/Media/ElectraCodecs/Source/ElectraDecoders/Private/ElectraDecodersUtils.cpp
+++ b/Engine/Plugins/Media/ElectraCodecs/Source/ElectraDecoders/Private/ElectraDecodersUtils.cpp
@@ -1,6 +1,7 @@
 // Copyright Epic Games, Inc. All Rights Reserved.
 
 #include "ElectraDecodersUtils.h"
+#include "Utilities/UtilitiesMP4.h"
 
 
 namespace ElectraDecodersUtil
@@ -304,6 +305,170 @@ bool ParseCodecVP9(FMimeTypeVideoCodecInfo& OutInfo, const FString& InCodecForma
 	return false;
 }
 
+bool ParseCodecAV1(FMimeTypeVideoCodecInfo& OutInfo, const FString& InCodecFormat, const TArray<uint8>& Inav1CBox, const TArray<uint8>& IncolrBox)
+{
+	if (InCodecFormat.StartsWith("av01"))
+	{
+		FString oti = InCodecFormat;
+		if (!Inav1CBox.IsEmpty())
+		{
+			// Enough data to represent the `av1C` box, and is it of version 1?
+			if (Inav1CBox.Num() < 4 || Inav1CBox[0] != 0x81)
+			{
+				return false;
+			}
+			auto high_bitdepth = [&]() -> bool { return (Inav1CBox[2] & 0x40) == 0x40; };
+			auto twelve_bit = [&]() -> bool { return (Inav1CBox[2] & 0x20) == 0x20; };
+
+			OutInfo.Codec = TEXT("av01");
+			OutInfo.Extras[0] = OutInfo.Profile = Inav1CBox[1] >> 5;	// profile
+			OutInfo.Extras[1] = OutInfo.Level = Inav1CBox[1] & 31;		// level
+			OutInfo.Extras[2] = OutInfo.Tier = Inav1CBox[2] >> 7;		// tier
+			uint8 NumBitsLuma = 8;
+			if (OutInfo.Profile == 2 && high_bitdepth())
+			{
+				NumBitsLuma = twelve_bit() ? 12 : 10;
+			}
+			else if (OutInfo.Profile <= 2)
+			{
+				NumBitsLuma = high_bitdepth() ? 10 : 8;
+			}
+			OutInfo.Extras[3] = OutInfo.NumBitsLuma = NumBitsLuma;		// bitDepth
+			OutInfo.Extras[4] = (Inav1CBox[2] >> 4) & 1;				// monochrome
+			OutInfo.Extras[5] = (Inav1CBox[2] >> 3) & 1;				// chromaSubsamplingX
+			OutInfo.Extras[6] = (Inav1CBox[2] >> 2) & 1;				// chromaSubsamplingY
+			OutInfo.Extras[7] = (Inav1CBox[2]) & 3;						// chromaSamplePosition
+			// Use defaults for values that have to come in from a 'colr' box.
+			OutInfo.Extras[8] = 1;										// colorPrimaries
+			OutInfo.Extras[9] = 1;										// transferCharacteristics
+			OutInfo.Extras[10] = 1;										// matrixCoefficients
+			OutInfo.Extras[11] = 0;										// videoFullRangeFlag
+
+			// Is a 'colr' box provided? And does it contain the mandated 'nclx' box?
+			Electra::UtilitiesMP4::FMP4AtomReader boxReader(IncolrBox);
+			uint32 Type;
+			if (boxReader.Read(Type) && Type == Electra::UtilitiesMP4::MakeBoxAtom('n','c','l','x'))
+			{
+				uint16 colour_primaries, transfer_characteristics, matrix_coeffs;
+				uint8 video_full_range_flag = 0;
+				if (boxReader.Read(colour_primaries) && boxReader.Read(transfer_characteristics) && boxReader.Read(matrix_coeffs) && boxReader.Read(video_full_range_flag))
+				{
+					OutInfo.Extras[8] = colour_primaries;
+					OutInfo.Extras[9] = transfer_characteristics;
+					OutInfo.Extras[10] = matrix_coeffs;
+					OutInfo.Extras[11] = video_full_range_flag;
+				}
+			}
+			return true;
+		}
+		if (oti.Len() > 4)
+		{
+			OutInfo.Codec = TEXT("av01");
+			TArray<FString> Parts;
+			InCodecFormat.ParseIntoArray(Parts, TEXT("."), false);
+			// Profile, level, tier and bitdepth are required.
+			if (Parts.Num() >= 4)
+			{
+				// Set defaults
+				OutInfo.Extras[4] = 0;			// monochrome
+				OutInfo.Extras[5] = 1;			// chromaSubsamplingX
+				OutInfo.Extras[6] = 1;			// chromaSubsamplingY
+				OutInfo.Extras[7] = 0;			// chromaSamplingPosition
+				OutInfo.Extras[8] = 1;			// colorPrimaries
+				OutInfo.Extras[9] = 1;			// transferCharacteristics
+				OutInfo.Extras[10] = 1;			// matrixCoefficients
+				OutInfo.Extras[11] = 0;			// videoFullRangeFlag
+				for(int32 i=1; i<Parts.Num(); ++i)
+				{
+					switch(i)
+					{
+						case 1:
+						{
+							// Profile. Has to be a single digit.
+							LexFromString(OutInfo.Profile, *Parts[i]);
+							OutInfo.Extras[0] = OutInfo.Profile;
+							break;
+						}
+						case 2:
+						{
+							// Level and tier. Needs to be 2 digits and an 'M' or 'H'.
+							if (Parts[i].Len() != 3 || (Parts[i][2]!=TCHAR('M') && Parts[i][2]!=TCHAR('m') && Parts[i][2]!=TCHAR('H') && Parts[i][2]!=TCHAR('h')))
+							{
+								return false;
+							}
+							OutInfo.Extras[2] = OutInfo.Tier = Parts[i][2]==TCHAR('H') || Parts[i][2]==TCHAR('h') ? 1 : 0;
+							Parts[i].LeftChopInline(1);
+							LexFromString(OutInfo.Level, *Parts[i]);
+							OutInfo.Extras[1] = OutInfo.Level;
+							break;
+						}
+						case 3:
+						{
+							// Bitdepth, 2 digits
+							LexFromString(OutInfo.NumBitsLuma, *Parts[i]);
+							OutInfo.Extras[3] = OutInfo.NumBitsLuma;
+							break;
+						}
+						/*
+							From this point on missing or malformed components are not critical. As per the standard:
+								"If any character that is not '.', digits, part of the AV1 4CC, or a tier value is encountered,
+								the string SHALL be interpreted ignoring all the characters starting from that character."
+						*/
+						case 4:
+						{
+							// Monochrome, 1 digit
+							LexFromString(OutInfo.Extras[4], *Parts[i]);
+							break;
+						}
+						case 5:
+						{
+							// Chroma subsampling, 3 digits
+							if (Parts[i].Len() != 3)
+							{
+								return true;
+							}
+							if (FChar::IsDigit(Parts[i][0])) OutInfo.Extras[5] = Parts[i][0]-TCHAR('0'); else return true;
+							if (FChar::IsDigit(Parts[i][1])) OutInfo.Extras[6] = Parts[i][1]-TCHAR('0'); else return true;
+							if (FChar::IsDigit(Parts[i][2])) OutInfo.Extras[7] = Parts[i][2]-TCHAR('0'); else return true;
+							break;
+						}
+						case 6:
+						{
+							// colorPrimaries.
+							LexFromString(OutInfo.Extras[8], *Parts[i]);
+							break;
+						}
+						case 7:
+						{
+							// transferCharacteristics.
+							LexFromString(OutInfo.Extras[9], *Parts[i]);
+							break;
+						}
+						case 8:
+						{
+							// matrixCoefficients.
+							LexFromString(OutInfo.Extras[10], *Parts[i]);
+							break;
+						}
+						case 9:
+						{
+							// videoFullRangeFlag.
+							LexFromString(OutInfo.Extras[11], *Parts[i]);
+							break;
+						}
+						default:
+						{
+							break;
+						}
+					}
+				}
+				return true;
+			}
+		}
+	}
+	return false;
+}
+
 
 int64 GetVariantValueSafeI64(const TMap<FString, FVariant>& InFromMap, const FString& InName, int64 InDefaultValue)
 {
diff --git a/Engine/Plugins/Media/ElectraCodecs/Source/ElectraDecoders/Private/Utils/AOMedia/ElectraBitstreamProcessor_AV1.cpp b/Engine/Plugins/Media/ElectraCodecs/Source/ElectraDecoders/Private/Utils/AOMedia/ElectraBitstreamProcessor_AV1.cpp
new file mode 100644
--- /dev/null
+++ b/Engine/Plugins/Media/ElectraCodecs/Source/ElectraDecoders/Private/Utils/AOMedia/ElectraBitstreamProcessor_AV1.cpp
@@ -0,0 +1,245 @@
+// Copyright Epic Games, Inc. All Rights Reserved.
+
+#include "Utils/AOMedia/ElectraBitstreamProcessor_AV1.h"
+#include "Utils/AOMedia/ElectraUtilsAV1Video.h"
+#include "Utils/MPEG/ElectraUtilsMPEGVideo.h"
+#include "IElectraDecoderFeaturesAndOptions.h"
+#include "ElectraDecodersUtils.h"
+
+
+class FElectraDecoderBitstreamProcessorAV1::FImpl
+{
+public:
+	class FBitstreamInfo : public IElectraDecoderBitstreamInfo
+	{
+	public:
+		virtual ~FBitstreamInfo() = default;
+	};
+
+	void SetPropertiesOnOutput(TMap<FString, FVariant>& InOutProperties)
+	{
+		if (CurrentColorimetry.IsValid() && bColorimetryChanged)
+		{
+			bColorimetryChanged = false;
+			TConstArrayView<uint8> Colorimetry = MakeConstArrayView(reinterpret_cast<const uint8*>(CurrentColorimetry.Get()), sizeof(ElectraDecodersUtil::MPEG::FCommonColorimetry));
+			InOutProperties.Emplace(IElectraDecoderBitstreamProcessorInfo::CommonColorimetry, TArray<uint8>(Colorimetry));
+		}
+
+		if (CurrentMDCV.IsValid() && bMDCVChanged)
+		{
+			bMDCVChanged = false;
+			TConstArrayView<uint8> Params = MakeConstArrayView(reinterpret_cast<const uint8*>(CurrentMDCV.Get()), sizeof(ElectraDecodersUtil::MPEG::FSEImastering_display_colour_volume));
+			InOutProperties.Emplace(IElectraDecoderBitstreamProcessorInfo::SeiMasteringDisplayColorVolume, TArray<uint8>(Params));
+		}
+
+		if (CurrentCLLI.IsValid() && bCLLIChanged)
+		{
+			bCLLIChanged = false;
+			TConstArrayView<uint8> Params = MakeConstArrayView(reinterpret_cast<const uint8*>(CurrentCLLI.Get()), sizeof(ElectraDecodersUtil::MPEG::FSEIcontent_light_level_info));
+			InOutProperties.Emplace(IElectraDecoderBitstreamProcessorInfo::SeiContentLightLeveInfo, TArray<uint8>(Params));
+		}
+	}
+
+	void UpdateColorimetry(uint8 colour_primaries, uint8 transfer_characteristics, uint8 matrix_coeffs, uint8 video_full_range_flag, uint8 video_format)
+	{
+		if (!CurrentColorimetry.IsValid() ||
+			CurrentColorimetry->colour_primaries != colour_primaries ||
+			CurrentColorimetry->transfer_characteristics != transfer_characteristics ||
+			CurrentColorimetry->matrix_coeffs != matrix_coeffs ||
+			CurrentColorimetry->video_full_range_flag != video_full_range_flag ||
+			CurrentColorimetry->video_format != video_format)
+		{
+			CurrentColorimetry = MakeShared<ElectraDecodersUtil::MPEG::FCommonColorimetry, ESPMode::ThreadSafe>();
+			CurrentColorimetry->colour_primaries = colour_primaries;
+			CurrentColorimetry->transfer_characteristics = transfer_characteristics;
+			CurrentColorimetry->matrix_coeffs = matrix_coeffs;
+			CurrentColorimetry->video_full_range_flag = video_full_range_flag;
+			CurrentColorimetry->video_format = video_format;
+			bColorimetryChanged = true;
+		}
+	}
+
+	void UpdateMDCV(const ElectraDecodersUtil::AV1Video::FMetadata_hdr_mdvc& InMDCV)
+	{
+		ElectraDecodersUtil::MPEG::FSEImastering_display_colour_volume mdcv;
+		mdcv.display_primaries_x[0] = InMDCV.primary_chromaticity_x[0];
+		mdcv.display_primaries_y[0] = InMDCV.primary_chromaticity_y[0];
+		mdcv.display_primaries_x[1] = InMDCV.primary_chromaticity_x[1];
+		mdcv.display_primaries_y[1] = InMDCV.primary_chromaticity_y[1];
+		mdcv.display_primaries_x[2] = InMDCV.primary_chromaticity_x[2];
+		mdcv.display_primaries_y[2] = InMDCV.primary_chromaticity_y[2];
+		mdcv.white_point_x = InMDCV.white_point_chromaticity_x;
+		mdcv.white_point_y = InMDCV.white_point_chromaticity_y;
+		mdcv.max_display_mastering_luminance = InMDCV.luminance_max;
+		mdcv.min_display_mastering_luminance = InMDCV.luminance_min;
+
+		if (!CurrentMDCV || *CurrentMDCV != mdcv)
+		{
+			TSharedPtr<ElectraDecodersUtil::MPEG::FSEImastering_display_colour_volume, ESPMode::ThreadSafe> NewMDCV = MakeShared<ElectraDecodersUtil::MPEG::FSEImastering_display_colour_volume, ESPMode::ThreadSafe>(mdcv);
+			CurrentMDCV = MoveTemp(NewMDCV);
+			bMDCVChanged = true;
+		}
+	}
+
+	void UpdateCLL(const ElectraDecodersUtil::AV1Video::FMetadata_hdr_cll& InCLL)
+	{
+		ElectraDecodersUtil::MPEG::FSEIcontent_light_level_info cll;
+		cll.max_content_light_level = InCLL.max_cll;
+		cll.max_pic_average_light_level = InCLL.max_fall;
+		if (!CurrentCLLI || *CurrentCLLI != cll)
+		{
+			TSharedPtr<ElectraDecodersUtil::MPEG::FSEIcontent_light_level_info, ESPMode::ThreadSafe> NewCLL = MakeShared<ElectraDecodersUtil::MPEG::FSEIcontent_light_level_info, ESPMode::ThreadSafe>(cll);
+			CurrentCLLI = MoveTemp(NewCLL);
+			bCLLIChanged = true;
+		}
+	}
+
+	void Clear()
+	{
+		CurrentColorimetry.Reset();
+		CurrentMDCV.Reset();
+		CurrentCLLI.Reset();
+		bColorimetryChanged = false;
+		bMDCVChanged = false;
+		bCLLIChanged = false;
+		LastErrorMessage.Empty();
+	}
+
+	static FString DCRName;
+	static FString CSDName;
+
+	TSharedPtr<ElectraDecodersUtil::MPEG::FCommonColorimetry, ESPMode::ThreadSafe> CurrentColorimetry;
+	TSharedPtr<ElectraDecodersUtil::MPEG::FSEImastering_display_colour_volume, ESPMode::ThreadSafe> CurrentMDCV;
+	TSharedPtr<ElectraDecodersUtil::MPEG::FSEIcontent_light_level_info, ESPMode::ThreadSafe> CurrentCLLI;
+
+	bool bColorimetryChanged = false;
+	bool bMDCVChanged = false;
+	bool bCLLIChanged = false;
+
+	FString LastErrorMessage;
+};
+FString FElectraDecoderBitstreamProcessorAV1::FImpl::DCRName(TEXT("dcr"));
+FString FElectraDecoderBitstreamProcessorAV1::FImpl::CSDName(TEXT("csd"));
+
+
+FElectraDecoderBitstreamProcessorAV1::~FElectraDecoderBitstreamProcessorAV1()
+{
+}
+
+FElectraDecoderBitstreamProcessorAV1::FElectraDecoderBitstreamProcessorAV1(const TMap<FString, FVariant>& InDecoderParams, const TMap<FString, FVariant>& InFormatParams)
+{
+	Impl = MakePimpl<FImpl>();
+}
+
+bool FElectraDecoderBitstreamProcessorAV1::WillModifyBitstreamInPlace()
+{
+	return false;
+}
+
+void FElectraDecoderBitstreamProcessorAV1::Clear()
+{
+	Impl->Clear();
+}
+
+IElectraDecoderBitstreamProcessor::EProcessResult FElectraDecoderBitstreamProcessorAV1::GetCSDFromConfigurationRecord(TArray<uint8>& OutCSD, const TMap<FString, FVariant>& InParamsWithDCRorCSD)
+{
+	OutCSD.Empty();
+	return EProcessResult::Ok;
+}
+
+IElectraDecoderBitstreamProcessor::EProcessResult FElectraDecoderBitstreamProcessorAV1::ProcessInputForDecoding(TSharedPtr<IElectraDecoderBitstreamInfo, ESPMode::ThreadSafe>& OutBSI, FElectraDecoderInputAccessUnit& InOutAccessUnit, const TMap<FString, FVariant>& InAccessUnitSidebandData)
+{
+	// Already processed?
+	if ((InOutAccessUnit.Flags & EElectraDecoderFlags::InputIsProcessed) != EElectraDecoderFlags::None)
+	{
+		return EProcessResult::Ok;
+	}
+
+	// Set to processed even if we fail somewhere now.
+	InOutAccessUnit.Flags |= EElectraDecoderFlags::InputIsProcessed;
+
+	// Go over the data and look at the OBUs
+	const uint8* CurrentOBU = (const uint8*)InOutAccessUnit.Data;
+	const uint8* EndOfOBUs = CurrentOBU + InOutAccessUnit.DataSize;
+	while(CurrentOBU < EndOfOBUs)
+	{
+		const uint8 OBUheaderByte = *CurrentOBU;
+		if ((OBUheaderByte & 0x80) != 0)
+		{
+			Impl->LastErrorMessage = FString::Printf(TEXT("obu_forbidden_bit not zero"));
+			return EProcessResult::Error;
+		}
+		const ElectraDecodersUtil::AV1Video::EOBUType OBUtype = (const ElectraDecodersUtil::AV1Video::EOBUType)(OBUheaderByte >> 3);
+		const uint8 obu_extension_flag = (OBUheaderByte >> 2) & 1;
+		const uint8 obu_has_size_field = (OBUheaderByte >> 1) & 1;
+		const uint8 obu_reserved_1bit = OBUheaderByte & 1;
+		CurrentOBU += obu_extension_flag ? 2 : 1;
+		uint32 obu_size;
+		uint32 nr = 0;
+		if (obu_has_size_field)
+		{
+			ElectraDecodersUtil::AV1Video::FBitstreamReaderAV1 bs(CurrentOBU, EndOfOBUs - CurrentOBU);
+			obu_size = bs.leb128(nr);
+			CurrentOBU += nr;
+		}
+		else
+		{
+			obu_size = EndOfOBUs - CurrentOBU;
+		}
+
+		if (OBUtype == ElectraDecodersUtil::AV1Video::EOBUType::OBU_SEQUENCE_HEADER)
+		{
+			ElectraDecodersUtil::AV1Video::FBitstreamReaderAV1 bs(CurrentOBU, obu_size);
+			ElectraDecodersUtil::AV1Video::FSequenceHeader sh;
+			if (!sh.ParseFrom(bs))
+			{
+				Impl->LastErrorMessage = FString::Printf(TEXT("Failed to parse sequence_header_obu"));
+				return EProcessResult::Error;
+			}
+			Impl->UpdateColorimetry((uint8)sh.color_config.color_primaries, (uint8)sh.color_config.transfer_characteristics, (uint8)sh.color_config.matrix_coefficients, sh.color_config.color_range, 5);
+		}
+		// HDR metadata?
+		else if (OBUtype == ElectraDecodersUtil::AV1Video::EOBUType::OBU_METADATA)
+		{
+			ElectraDecodersUtil::AV1Video::FBitstreamReaderAV1 bs(CurrentOBU, obu_size);
+			const ElectraDecodersUtil::AV1Video::EMetadataType MetadataType = (const ElectraDecodersUtil::AV1Video::EMetadataType)(bs.leb128(nr));
+			switch(MetadataType)
+			{
+				case ElectraDecodersUtil::AV1Video::EMetadataType::METADATA_TYPE_HDR_CLL:
+				{
+					ElectraDecodersUtil::AV1Video::FMetadata_hdr_cll cll;
+					if (cll.ParseFrom(bs))
+					{
+						Impl->UpdateCLL(cll);
+					}
+					break;
+				}
+				case ElectraDecodersUtil::AV1Video::EMetadataType::METADATA_TYPE_HDR_MDCV:
+				{
+					ElectraDecodersUtil::AV1Video::FMetadata_hdr_mdvc mdcv;
+					if (mdcv.ParseFrom(bs))
+					{
+						Impl->UpdateMDCV(mdcv);
+					}
+					break;
+				}
+			}
+		}
+		CurrentOBU += obu_size;
+	}
+
+	return EProcessResult::Ok;
+}
+
+void FElectraDecoderBitstreamProcessorAV1::SetPropertiesOnOutput(TMap<FString, FVariant>& InOutProperties, TSharedPtr<IElectraDecoderBitstreamInfo, ESPMode::ThreadSafe> InBSI)
+{
+	Impl->SetPropertiesOnOutput(InOutProperties);
+//	uint8 num_bits = 8;
+//	InOutProperties.Emplace(IElectraDecoderBitstreamProcessorInfo::NumBitsLuma, FVariant(num_bits));
+}
+
+FString FElectraDecoderBitstreamProcessorAV1::GetLastError()
+{
+	return Impl->LastErrorMessage;
+}
+
diff --git a/Engine/Plugins/Media/ElectraCodecs/Source/ElectraDecoders/Private/Utils/AOMedia/ElectraUtilsAV1Video.cpp b/Engine/Plugins/Media/ElectraCodecs/Source/ElectraDecoders/Private/Utils/AOMedia/ElectraUtilsAV1Video.cpp
new file mode 100644
--- /dev/null
+++ b/Engine/Plugins/Media/ElectraCodecs/Source/ElectraDecoders/Private/Utils/AOMedia/ElectraUtilsAV1Video.cpp
@@ -0,0 +1,292 @@
+// Copyright Epic Games, Inc. All Rights Reserved.
+
+#include "Utils/AOMedia/ElectraUtilsAV1Video.h"
+#include "Utilities/ElectraBitstream.h"
+
+
+namespace ElectraDecodersUtil
+{
+	namespace AV1Video
+	{
+		bool FMetadata_hdr_cll::ParseFrom(FBitstreamReaderAV1& InReader)
+		{
+			if (InReader.GetRemainingBits() < 32)
+			{
+				return false;
+			}
+			max_cll = InReader.GetBits(16);
+			max_fall = InReader.GetBits(16);
+			return true;
+		}
+
+		bool FMetadata_hdr_mdvc::ParseFrom(FBitstreamReaderAV1& InReader)
+		{
+			if (InReader.GetRemainingBits() < 192)
+			{
+				return false;
+			}
+			primary_chromaticity_x[0] = InReader.GetBits(16);
+			primary_chromaticity_y[0] = InReader.GetBits(16);
+			primary_chromaticity_x[1] = InReader.GetBits(16);
+			primary_chromaticity_y[1] = InReader.GetBits(16);
+			primary_chromaticity_x[2] = InReader.GetBits(16);
+			primary_chromaticity_y[2] = InReader.GetBits(16);
+			white_point_chromaticity_x = InReader.GetBits(16);
+			white_point_chromaticity_y = InReader.GetBits(16);
+			luminance_max = InReader.GetBits(32);
+			luminance_min = InReader.GetBits(32);
+			return true;
+		}
+
+
+
+		bool FSequenceHeader::ParseFrom(FBitstreamReaderAV1& InReader)
+		{
+			seq_profile = (uint8) InReader.GetBits(3);
+			if (seq_profile > 2)
+			{
+				return false;
+			}
+			still_picture = InReader.GetBits(1);
+			reduced_still_picture_header = InReader.GetBits(1);
+			if (reduced_still_picture_header && !still_picture)
+			{
+				return false;
+			}
+			decoder_model_info_present_flag = 0;
+			initial_display_delay_present_flag = 0;
+			operating_points_cnt_minus_1 = 0;
+			if (!reduced_still_picture_header)
+			{
+				timing_info_present_flag = InReader.GetBits(1);
+				if (timing_info_present_flag)
+				{
+					// timing_info
+					timing_info.num_units_in_display_tick = InReader.GetBits(32);
+					timing_info.time_scale = InReader.GetBits(32);
+					timing_info.equal_picture_interval = InReader.GetBits(1);
+					if (timing_info.equal_picture_interval)
+					{
+						timing_info.num_ticks_per_picture_minus_1 = InReader.uvlc();
+						if (timing_info.num_ticks_per_picture_minus_1 == 0xffffffffU)
+						{
+							return false;
+						}
+					}
+					decoder_model_info_present_flag = InReader.GetBits(1);
+					if (decoder_model_info_present_flag)
+					{
+						// decoder_model_info
+						decoder_model_info.buffer_delay_length_minus_1 = InReader.GetBits(5);
+						decoder_model_info.num_units_in_decoding_tick = InReader.GetBits(32);
+						decoder_model_info.buffer_removal_time_length_minus_1 = InReader.GetBits(5);
+						decoder_model_info.frame_presentation_time_length_minus_1 = InReader.GetBits(5);
+					}
+				}
+
+				initial_display_delay_present_flag = InReader.GetBits(1);
+				operating_points_cnt_minus_1 = InReader.GetBits(5);
+				for(int32 i=0; i<=operating_points_cnt_minus_1; ++i)
+				{
+					operating_points[i].operating_point_idc = InReader.GetBits(12);
+					operating_points[i].seq_level_idx = InReader.GetBits(5);
+					operating_points[i].seq_tier = 0;
+					if (operating_points[i].seq_level_idx > 7)
+					{
+						operating_points[i].seq_tier = InReader.GetBits(1);
+					}
+					if (decoder_model_info_present_flag)
+					{
+						operating_points[i].decoder_model_present_for_this_op = InReader.GetBits(1);
+						if (operating_points[i].decoder_model_present_for_this_op)
+						{
+							// operating_parameters_info
+							uint32 nb = decoder_model_info.buffer_delay_length_minus_1 + 1;
+							operating_parameters_info[i].decoder_buffer_delay = InReader.GetBits(nb);
+							operating_parameters_info[i].encoder_buffer_delay = InReader.GetBits(nb);
+							operating_parameters_info[i].low_delay_mode_flag = InReader.GetBits(1);
+						}
+					}
+					else
+					{
+						operating_points[i].decoder_model_present_for_this_op = 0;
+					}
+					operating_points[i].initial_display_delay_present_for_this_op = 0;
+					operating_points[i].initial_display_delay_minus_1 = 0;
+					if (initial_display_delay_present_flag)
+					{
+						operating_points[i].initial_display_delay_present_for_this_op = InReader.GetBits(1);
+						if (operating_points[i].initial_display_delay_present_for_this_op)
+						{
+							operating_points[i].initial_display_delay_minus_1 = InReader.GetBits(4);
+						}
+					}
+				}
+			}
+			else
+			{
+				timing_info_present_flag = 0;
+				operating_points[0].operating_point_idc = 0;
+				operating_points[0].seq_level_idx = InReader.GetBits(5);
+				operating_points[0].seq_tier = 0;
+				operating_points[0].decoder_model_present_for_this_op = 0;
+				operating_points[0].initial_display_delay_present_for_this_op = 0;
+			}
+
+			//operatingPoint = choose_operating_point();
+			//OperatingPointIdc = operating_point_idc[operatingPoint];
+
+			frame_width_bits_minus_1 = InReader.GetBits(4);
+			frame_height_bits_minus_1 = InReader.GetBits(4);
+			max_frame_width_minus_1 = InReader.GetBits(frame_width_bits_minus_1 + 1);
+			max_frame_height_minus_1 = InReader.GetBits(frame_height_bits_minus_1 + 1);
+			frame_id_numbers_present_flag = 0;
+			if (!reduced_still_picture_header)
+			{
+				frame_id_numbers_present_flag = InReader.GetBits(1);
+			}
+			if (frame_id_numbers_present_flag)
+			{
+				delta_frame_id_length_minus_2 = InReader.GetBits(4);
+				additional_frame_id_length_minus_1 = InReader.GetBits(3);
+			}
+			use_128x128_superblock = InReader.GetBits(1);
+			enable_filter_intra = InReader.GetBits(1);
+			enable_intra_edge_filter = InReader.GetBits(1);
+
+			const uint8 SELECT_SCREEN_CONTENT_TOOLS = 2;
+			const uint8 SELECT_INTEGER_MV = 2;
+
+			enable_interintra_compound = 0;
+			enable_masked_compound = 0;
+			enable_warped_motion = 0;
+			enable_dual_filter = 0;
+			enable_order_hint = 0;
+			enable_jnt_comp = 0;
+			enable_ref_frame_mvs = 0;
+			seq_force_screen_content_tools = SELECT_SCREEN_CONTENT_TOOLS;
+			seq_force_integer_mv = SELECT_INTEGER_MV;
+			OrderHintBits = 0;
+			if (!reduced_still_picture_header)
+			{
+				enable_interintra_compound = InReader.GetBits(1);
+				enable_masked_compound = InReader.GetBits(1);
+				enable_warped_motion = InReader.GetBits(1);
+				enable_dual_filter = InReader.GetBits(1);
+				enable_order_hint = InReader.GetBits(1);
+				enable_jnt_comp = 0;
+				enable_ref_frame_mvs = 0;
+				if (enable_order_hint)
+				{
+					enable_jnt_comp = InReader.GetBits(1);
+					enable_ref_frame_mvs = InReader.GetBits(1);
+				}
+				seq_choose_screen_content_tools = InReader.GetBits(1);
+				seq_force_screen_content_tools = SELECT_SCREEN_CONTENT_TOOLS;
+				if (!seq_choose_screen_content_tools)
+				{
+					seq_force_screen_content_tools = InReader.GetBits(1);
+				}
+				seq_force_integer_mv = SELECT_INTEGER_MV;
+				if (seq_force_screen_content_tools)
+				{
+					seq_choose_integer_mv = InReader.GetBits(1);
+					if (!seq_choose_integer_mv)
+					{
+						seq_force_integer_mv = InReader.GetBits(1);
+					}
+				}
+				if (enable_order_hint)
+				{
+					order_hint_bits_minus_1 = InReader.GetBits(3);
+				}
+			}
+
+			enable_superres = InReader.GetBits(1);
+			enable_cdef = InReader.GetBits(1);
+			enable_restoration = InReader.GetBits(1);
+
+			// color_config
+			int32 BitDepth = 8;
+			color_config.high_bitdepth = InReader.GetBits(1);
+			if (seq_profile == 2 && color_config.high_bitdepth)
+			{
+				color_config.twelve_bit = InReader.GetBits(1);
+				BitDepth = color_config.twelve_bit ? 12 : 10;
+			}
+			else if (seq_profile <= 2)
+			{
+				BitDepth = color_config.high_bitdepth ? 10 : 8;
+			}
+			color_config.mono_chrome = seq_profile == 1 ? 0 : InReader.GetBits(1);
+			color_config.color_description_present_flag = InReader.GetBits(1);
+			if (color_config.color_description_present_flag)
+			{
+				color_config.color_primaries = (EColorPrimaries) InReader.GetBits(8);
+				color_config.transfer_characteristics = (ETransferCharacteristics) InReader.GetBits(8);
+				color_config.matrix_coefficients = (EMatrixCoefficients) InReader.GetBits(8);
+			}
+			else
+			{
+				color_config.color_primaries = EColorPrimaries::CP_UNSPECIFIED;
+				color_config.transfer_characteristics = ETransferCharacteristics::TC_UNSPECIFIED;
+				color_config.matrix_coefficients = EMatrixCoefficients::MC_UNSPECIFIED;
+			}
+			if (color_config.mono_chrome)
+			{
+				color_config.color_range = InReader.GetBits(1);
+				color_config.subsampling_x = 1;
+				color_config.subsampling_y = 1;
+				color_config.chroma_sample_position = EChromaSamplePosition::CSP_UNKNOWN;
+				color_config.separate_uv_delta_q = 0;
+			}
+			else if (color_config.color_primaries == EColorPrimaries::CP_BT_709 && color_config.transfer_characteristics == ETransferCharacteristics::TC_SRGB && color_config.matrix_coefficients == EMatrixCoefficients::MC_IDENTITY)
+			{
+				color_config.color_range = 1;
+				color_config.subsampling_x = 0;
+				color_config.subsampling_y = 0;
+			}
+			else
+			{
+				color_config.color_range = InReader.GetBits(1);
+				if (seq_profile == 0)
+				{
+					color_config.subsampling_x = 1;
+					color_config.subsampling_y = 1;
+				}
+				else if (seq_profile == 1)
+				{
+					color_config.subsampling_x = 0;
+					color_config.subsampling_y = 0;
+				}
+				else
+				{
+					if (BitDepth == 12)
+					{
+						color_config.subsampling_x = InReader.GetBits(1);
+						color_config.subsampling_y = color_config.subsampling_x ? InReader.GetBits(1) : 0;
+					}
+					else
+					{
+						color_config.subsampling_x = 1;
+						color_config.subsampling_y = 0;
+					}
+				}
+				if (color_config.subsampling_x && color_config.subsampling_y)
+				{
+					color_config.chroma_sample_position = (EChromaSamplePosition) InReader.GetBits(2);
+				}
+			}
+			if (color_config.mono_chrome == 0)
+			{
+				color_config.separate_uv_delta_q = InReader.GetBits(1);
+			}
+
+			film_grain_params_present = InReader.GetBits(1);
+			return true;
+		}
+
+
+	} // namespace AV1Video
+
+} // namespace ElectraDecodersUtil
diff --git a/Engine/Plugins/Media/ElectraCodecs/Source/ElectraDecoders/Public/ElectraDecodersUtils.h b/Engine/Plugins/Media/ElectraCodecs/Source/ElectraDecoders/Public/ElectraDecodersUtils.h
--- a/Engine/Plugins/Media/ElectraCodecs/Source/ElectraDecoders/Public/ElectraDecodersUtils.h
+++ b/Engine/Plugins/Media/ElectraCodecs/Source/ElectraDecoders/Public/ElectraDecodersUtils.h
@@ -67,7 +67,7 @@ namespace ElectraDecodersUtil
 		uint64 Constraints = 0;
 		int32 Tier = 0;
 		int32 NumBitsLuma = 0;
-		int32 Extras[8] {0};
+		int32 Extras[12] {0};
 	};
 
 
@@ -94,6 +94,7 @@ namespace ElectraDecodersUtil
 	bool ELECTRADECODERS_API ParseCodecH265(FMimeTypeVideoCodecInfo& OutInfo, const FString& InCodecFormat);
 	bool ELECTRADECODERS_API ParseCodecVP8(FMimeTypeVideoCodecInfo& OutInfo, const FString& InCodecFormat, const TArray<uint8>& InvpcCBox);
 	bool ELECTRADECODERS_API ParseCodecVP9(FMimeTypeVideoCodecInfo& OutInfo, const FString& InCodecFormat, const TArray<uint8>& InvpcCBox);
+	bool ELECTRADECODERS_API ParseCodecAV1(FMimeTypeVideoCodecInfo& OutInfo, const FString& InCodecFormat, const TArray<uint8>& Inav1CBox, const TArray<uint8>& IncolrBox);
 
 
 	int64 ELECTRADECODERS_API GetVariantValueSafeI64(const TMap<FString, FVariant>& InFromMap, const FString& InName, int64 InDefaultValue=0);
diff --git a/Engine/Plugins/Media/ElectraCodecs/Source/ElectraDecoders/Public/Utils/AOMedia/ElectraBitstreamProcessor_AV1.h b/Engine/Plugins/Media/ElectraCodecs/Source/ElectraDecoders/Public/Utils/AOMedia/ElectraBitstreamProcessor_AV1.h
new file mode 100644
--- /dev/null
+++ b/Engine/Plugins/Media/ElectraCodecs/Source/ElectraDecoders/Public/Utils/AOMedia/ElectraBitstreamProcessor_AV1.h
@@ -0,0 +1,26 @@
+// Copyright Epic Games, Inc. All Rights Reserved.
+
+#pragma once
+
+#include "IElectraDecoder.h"
+
+
+class FElectraDecoderBitstreamProcessorAV1 : public IElectraDecoderBitstreamProcessor
+{
+public:
+	static TSharedPtr<IElectraDecoderBitstreamProcessor, ESPMode::ThreadSafe> Create(const TMap<FString, FVariant>& InDecoderParams, const TMap<FString, FVariant>& InFormatParams)
+	{ return MakeShareable<>(new FElectraDecoderBitstreamProcessorAV1(InDecoderParams, InFormatParams)); }
+
+	virtual ~FElectraDecoderBitstreamProcessorAV1();
+	bool WillModifyBitstreamInPlace() override;
+	void Clear() override;
+	EProcessResult GetCSDFromConfigurationRecord(TArray<uint8>& OutCSD, const TMap<FString, FVariant>& InParamsWithDCRorCSD) override;
+	EProcessResult ProcessInputForDecoding(TSharedPtr<IElectraDecoderBitstreamInfo, ESPMode::ThreadSafe>& OutBSI, FElectraDecoderInputAccessUnit& InOutAccessUnit, const TMap<FString, FVariant>& InAccessUnitSidebandData) override;
+	void SetPropertiesOnOutput(TMap<FString, FVariant>& InOutProperties, TSharedPtr<IElectraDecoderBitstreamInfo, ESPMode::ThreadSafe> InBSI) override;
+	FString GetLastError() override;
+
+private:
+	ELECTRADECODERS_API FElectraDecoderBitstreamProcessorAV1(const TMap<FString, FVariant>& InDecoderParams, const TMap<FString, FVariant>& InFormatParams);
+	class FImpl;
+	TPimplPtr<FImpl> Impl;
+};
diff --git a/Engine/Plugins/Media/ElectraCodecs/Source/ElectraDecoders/Public/Utils/AOMedia/ElectraUtilsAV1Video.h b/Engine/Plugins/Media/ElectraCodecs/Source/ElectraDecoders/Public/Utils/AOMedia/ElectraUtilsAV1Video.h
new file mode 100644
--- /dev/null
+++ b/Engine/Plugins/Media/ElectraCodecs/Source/ElectraDecoders/Public/Utils/AOMedia/ElectraUtilsAV1Video.h
@@ -0,0 +1,281 @@
+// Copyright Epic Games, Inc. All Rights Reserved.
+
+#pragma once
+
+#include <CoreMinimal.h>
+#include <Containers/Array.h>
+#include "Utilities/ElectraBitstream.h"
+
+namespace ElectraDecodersUtil
+{
+	namespace AV1Video
+	{
+
+		class FBitstreamReaderAV1 : public Electra::FBitstreamReader
+		{
+		public:
+			FBitstreamReaderAV1() = default;
+			FBitstreamReaderAV1(const uint8* InData, uint64 InDataSize)
+				: Electra::FBitstreamReader(InData, InDataSize)
+			{ }
+
+			uint32 uvlc()
+			{
+				uint32 leadingZeros = 0;
+				while(!GetBits(1))
+				{
+					if (++leadingZeros == 32)
+					{
+						return 0xffffffffU;
+					}
+				}
+				uint32 v = GetBits(leadingZeros);
+				return v + ((1U << leadingZeros) - 1U);
+			}
+
+			uint8 GetByte()
+			{
+				check(IsByteAligned());
+				return (DataSize - BytePosition) ? Data[BytePosition++] : 0U;
+			}
+
+			uint32 leb128(uint32& OutNumBytesRead)
+			{
+				uint64 v = 0;
+				uint32 Shift = 0;
+				bool bMore = false;
+				OutNumBytesRead = 0;
+				do
+				{
+					uint8 lb = GetByte();
+					bMore = !!(lb & 0x80);
+					v |= (uint64)(lb & 0x7fU) << Shift;
+					Shift += 7;
+					++OutNumBytesRead;
+				} while(bMore && Shift < 56);
+				if (v > 0xffffffffU || bMore)
+				{
+					// Not good.
+					return 0;
+				}
+				return (uint32)(v & 0xffffffffU);
+			}
+			uint32 leb128()
+			{
+				uint32 nr;
+				return leb128(nr);
+			}
+		};
+
+
+		enum class EOBUType
+		{
+			Reserved0 = 0,
+			OBU_SEQUENCE_HEADER,
+			OBU_TEMPORAL_DELIMITER,
+			OBU_FRAME_HEADER,
+			OBU_TILE_GROUP,
+			OBU_METADATA,
+			OBU_FRAME,
+			OBU_REDUNDANT_FRAME_HEADER,
+			OBU_TILE_LIST,
+			Reserved9 = 9,
+			OBU_PADDING = 15
+		};
+
+		enum class EMetadataType
+		{
+			Reserved0 = 0,
+			METADATA_TYPE_HDR_CLL,
+			METADATA_TYPE_HDR_MDCV,
+			METADATA_TYPE_SCALABILITY,
+			METADATA_TYPE_ITUT_T35,
+			METADATA_TYPE_TIMECODE,
+			UnregisteredUserPrivate6 = 6,
+			Reserved32 = 32
+		};
+
+		enum class EColorPrimaries : uint8
+		{
+			Undefined = 0,
+			CP_BT_709 = 1,
+			CP_UNSPECIFIED = 2,
+			CP_BT_470_M = 4,
+			CP_BT_470_B_G,
+			CP_BT_601,
+			CP_SMPTE_240,
+			CP_GENERIC_FILM,
+			CP_BT_2020,
+			CP_XYZ,
+			CP_SMPTE_431,
+			CP_SMPTE_432,
+			CP_EBU_3213 = 22
+		};
+
+		enum class ETransferCharacteristics : uint8
+		{
+			TC_RESERVED_0 = 0,
+			TC_BT_709,
+			TC_UNSPECIFIED,
+			TC_RESERVED_3,
+			TC_BT_470_M,
+			TC_BT_470_B_G,
+			TC_BT_601,
+			TC_SMPTE_240,
+			TC_LINEAR,
+			TC_LOG_100,
+			TC_LOG_100_SQRT10,
+			TC_IEC_61966,
+			TC_BT_1361,
+			TC_SRGB,
+			TC_BT_2020_10_BIT,
+			TC_BT_2020_12_BIT,
+			TC_SMPTE_2084,
+			TC_SMPTE_428,
+			TC_HLG
+		};
+
+		enum class EMatrixCoefficients : uint8
+		{
+			MC_IDENTITY = 0,
+			MC_BT_709,
+			MC_UNSPECIFIED,
+			MC_RESERVED_3,
+			MC_FCC,
+			MC_BT_470_B_G,
+			MC_BT_601,
+			MC_SMPTE_240,
+			MC_SMPTE_YCGCO,
+			MC_BT_2020_NCL,
+			MC_BT_2020_CL,
+			MC_SMPTE_2085,
+			MC_CHROMAT_NCL,
+			MC_CHROMAT_CL,
+			MC_ICTCP
+		};
+
+		enum class EChromaSamplePosition : uint8
+		{
+			CSP_UNKNOWN = 0,
+			CSP_VERTICAL,
+			CSP_COLOCATED,
+			CSP_RESERVED
+		};
+
+		struct FMetadata_hdr_cll
+		{
+			uint16 max_cll = 0;
+			uint16 max_fall = 0;
+			bool ELECTRADECODERS_API ParseFrom(FBitstreamReaderAV1& InReader);
+		};
+
+		struct FMetadata_hdr_mdvc
+		{
+			uint16 primary_chromaticity_x[3] { 0 };	// CIE 1931; 0=red, 1=green, 2=blue
+			uint16 primary_chromaticity_y[3] { 0 };
+			uint16 white_point_chromaticity_x = 0;
+			uint16 white_point_chromaticity_y = 0;
+			uint32 luminance_max = 0;				// a 24.8 fixed-point maximum luminance, represented in candelas per square meter
+			uint32 luminance_min = 0;				// a 18.14 fixed-point minimum luminance, represented in candelas per square meter
+			bool ELECTRADECODERS_API ParseFrom(FBitstreamReaderAV1& InReader);
+		};
+
+
+		struct FSequenceHeader
+		{
+			uint8 seq_profile = 0;
+			uint8 still_picture = 0;
+			uint8 reduced_still_picture_header = 0;
+			uint8 timing_info_present_flag = 0;
+			struct FTimingInfo
+			{
+				uint32 num_units_in_display_tick = 0;
+				uint32 time_scale = 0;
+				uint8 equal_picture_interval = 0;
+				uint32 num_ticks_per_picture_minus_1 = 0;
+			};
+			FTimingInfo timing_info;
+
+			uint8 decoder_model_info_present_flag = 0;
+			struct FDecoderModelInfo
+			{
+				uint8 buffer_delay_length_minus_1 = 0;
+				uint32 num_units_in_decoding_tick = 0;
+				uint8 buffer_removal_time_length_minus_1 = 0;
+				uint8 frame_presentation_time_length_minus_1 = 0;
+			};
+			FDecoderModelInfo decoder_model_info;
+
+			uint8 initial_display_delay_present_flag = 0;
+			uint8 operating_points_cnt_minus_1 = 0;
+
+			struct FOperatingPoint
+			{
+				uint16 operating_point_idc = 0;
+				uint8 seq_level_idx = 0;
+				uint8 seq_tier = 0;
+				uint8 decoder_model_present_for_this_op = 0;
+				uint8 initial_display_delay_present_for_this_op = 0;
+				uint8 initial_display_delay_minus_1 = 0;
+			};
+			FOperatingPoint operating_points[32];
+			struct FOperatingParametersInfo
+			{
+				uint32 decoder_buffer_delay = 0;
+				uint32 encoder_buffer_delay = 0;
+				uint8 low_delay_mode_flag = 0;
+			};
+			FOperatingParametersInfo operating_parameters_info[32];
+
+			uint8 frame_width_bits_minus_1 = 0;
+			uint8 frame_height_bits_minus_1 = 0;
+			uint32 max_frame_width_minus_1 = 0;
+			uint32 max_frame_height_minus_1 = 0;
+			uint8 frame_id_numbers_present_flag = 0;
+			uint8 delta_frame_id_length_minus_2 = 0;
+			uint8 additional_frame_id_length_minus_1 = 0;
+			uint8 use_128x128_superblock = 0;
+			uint8 enable_filter_intra = 0;
+			uint8 enable_intra_edge_filter = 0;
+			uint8 enable_interintra_compound = 0;
+			uint8 enable_masked_compound = 0;
+			uint8 enable_warped_motion = 0;
+			uint8 enable_dual_filter = 0;
+			uint8 enable_order_hint = 0;
+			uint8 enable_jnt_comp = 0;
+			uint8 enable_ref_frame_mvs = 0;
+			uint8 seq_choose_screen_content_tools = 0;
+			uint8 seq_force_screen_content_tools = 0;
+			uint8 seq_force_integer_mv = 0;
+			uint8 seq_choose_integer_mv = 0;
+			uint8 order_hint_bits_minus_1 = 0;
+			uint8 OrderHintBits = 0;
+			uint8 enable_superres = 0;
+			uint8 enable_cdef = 0;
+			uint8 enable_restoration = 0;
+
+			struct FColorConfig
+			{
+				uint8 high_bitdepth = 0;
+				uint8 twelve_bit = 0;
+				uint8 mono_chrome = 0;
+				uint8 color_description_present_flag = 0;
+				EColorPrimaries color_primaries = EColorPrimaries::CP_UNSPECIFIED;
+				ETransferCharacteristics transfer_characteristics = ETransferCharacteristics::TC_UNSPECIFIED;
+				EMatrixCoefficients matrix_coefficients = EMatrixCoefficients::MC_UNSPECIFIED;
+				uint8 color_range = 0;
+				uint8 subsampling_x = 0;
+				uint8 subsampling_y = 0;
+				EChromaSamplePosition chroma_sample_position = EChromaSamplePosition::CSP_UNKNOWN;
+				uint8 separate_uv_delta_q = 0;
+			};
+			FColorConfig color_config;
+
+			uint8 film_grain_params_present = 0;
+
+
+			bool ELECTRADECODERS_API ParseFrom(FBitstreamReaderAV1& InReader);
+		};
+
+	} // namespace AV1Video
+} // namespace ElectraDecodersUtil
diff --git a/Engine/Plugins/Media/ElectraCodecs/Source/ElectraDecoders/Public/Utils/MPEG/ElectraUtilsMPEGVideo.h b/Engine/Plugins/Media/ElectraCodecs/Source/ElectraDecoders/Public/Utils/MPEG/ElectraUtilsMPEGVideo.h
--- a/Engine/Plugins/Media/ElectraCodecs/Source/ElectraDecoders/Public/Utils/MPEG/ElectraUtilsMPEGVideo.h
+++ b/Engine/Plugins/Media/ElectraCodecs/Source/ElectraDecoders/Public/Utils/MPEG/ElectraUtilsMPEGVideo.h
@@ -47,6 +47,23 @@ namespace ElectraDecodersUtil
 			uint16 white_point_y = 0;
 			uint32 max_display_mastering_luminance = 0;
 			uint32 min_display_mastering_luminance = 0;
+			bool operator==(const FSEImastering_display_colour_volume& rhs) const
+			{
+				return display_primaries_x[0] == rhs.display_primaries_x[0] &&
+					   display_primaries_y[0] == rhs.display_primaries_y[0] &&
+					   display_primaries_x[1] == rhs.display_primaries_x[1] &&
+					   display_primaries_y[1] == rhs.display_primaries_y[1] &&
+					   display_primaries_x[2] == rhs.display_primaries_x[2] &&
+					   display_primaries_y[2] == rhs.display_primaries_y[2] &&
+					   white_point_x == rhs.white_point_x &&
+					   white_point_y == rhs.white_point_y &&
+					   max_display_mastering_luminance == rhs.max_display_mastering_luminance &&
+					   min_display_mastering_luminance == rhs.min_display_mastering_luminance;
+			}
+			bool operator!=(const FSEImastering_display_colour_volume& rhs) const
+			{
+				return !operator==(rhs);
+			}
 		};
 		bool ELECTRADECODERS_API ParseFromSEIMessage(FSEImastering_display_colour_volume& OutMDCV, const FSEIMessage& InMessage);
 		bool ELECTRADECODERS_API ParseFromMDCVBox(FSEImastering_display_colour_volume& OutMDCV, const TConstArrayView<uint8>& InMDCVBox);
@@ -56,6 +73,14 @@ namespace ElectraDecodersUtil
 		{
 			uint16 max_content_light_level = 0;
 			uint16 max_pic_average_light_level = 0;
+			bool operator==(const FSEIcontent_light_level_info& rhs) const
+			{
+				return max_content_light_level == rhs.max_content_light_level && max_pic_average_light_level == rhs.max_pic_average_light_level;
+			}
+			bool operator!=(const FSEIcontent_light_level_info& rhs) const
+			{
+				return !operator==(rhs);
+			}
 		};
 		bool ELECTRADECODERS_API ParseFromSEIMessage(FSEIcontent_light_level_info& OutCLLI, const FSEIMessage& InMessage);
 		bool ELECTRADECODERS_API ParseFromCOLLBox(FSEIcontent_light_level_info& OutCLLI, const TConstArrayView<uint8>& InCOLLBox);
diff --git a/Engine/Plugins/Media/ElectraUtil/Source/ElectraBase/Public/Utilities/ElectraBitstream.h b/Engine/Plugins/Media/ElectraUtil/Source/ElectraBase/Public/Utilities/ElectraBitstream.h
--- a/Engine/Plugins/Media/ElectraUtil/Source/ElectraBase/Public/Utilities/ElectraBitstream.h
+++ b/Engine/Plugins/Media/ElectraUtil/Source/ElectraBase/Public/Utilities/ElectraBitstream.h
@@ -212,7 +212,7 @@ public:
 		}
 	}
 
-private:
+protected:
 	const uint8*	Data = nullptr;
 	uint64			DataSize = 0;
 	uint64			BytePosition = 0;
