ElectraCodecs: Added a modular feature to fill codec information in a single location; changed Protron to use it
#rnx

#Synced-CL 0

diff --git a/Engine/Plugins/Media/ElectraCodecs/Source/ElectraDecoders/Private/ElectraDecodersModule.cpp b/Engine/Plugins/Media/ElectraCodecs/Source/ElectraDecoders/Private/ElectraDecodersModule.cpp
--- a/Engine/Plugins/Media/ElectraCodecs/Source/ElectraDecoders/Private/ElectraDecodersModule.cpp
+++ b/Engine/Plugins/Media/ElectraCodecs/Source/ElectraDecoders/Private/ElectraDecodersModule.cpp
@@ -8,6 +8,8 @@
 #include "IElectraDecodersModule.h"
 #include "PlatformElectraDecoders.h"
 
+#include "Utils/ElectraCodecFormatUtils.h"
+
 // Common codecs
 #include "ElectraMediaMP3Decoder.h"
 
@@ -21,6 +23,7 @@ class FElectraDecodersModule : public IElectraDecodersModule
 public:
 	void StartupModule() override
 	{
+		FElectraCodecFormatUtilsModularFeature::Startup();
 		FPlatformElectraDecoders::Startup();
 		FElectraMediaMP3Decoder::Startup();
 	}
@@ -29,6 +32,7 @@ public:
 	{
 		FElectraMediaMP3Decoder::Shutdown();
 		FPlatformElectraDecoders::Shutdown();
+		FElectraCodecFormatUtilsModularFeature::Shutdown();
 	}
 
 	bool SupportsDynamicReloading() override
diff --git a/Engine/Plugins/Media/ElectraCodecs/Source/ElectraDecoders/Private/ElectraDecodersUtils.cpp b/Engine/Plugins/Media/ElectraCodecs/Source/ElectraDecoders/Private/ElectraDecodersUtils.cpp
deleted file mode 100644
diff --git a/Engine/Plugins/Media/ElectraCodecs/Source/ElectraDecoders/Private/Utils/AOMedia/ElectraUtilsAV1Video.cpp b/Engine/Plugins/Media/ElectraCodecs/Source/ElectraDecoders/Private/Utils/AOMedia/ElectraUtilsAV1Video.cpp
--- a/Engine/Plugins/Media/ElectraCodecs/Source/ElectraDecoders/Private/Utils/AOMedia/ElectraUtilsAV1Video.cpp
+++ b/Engine/Plugins/Media/ElectraCodecs/Source/ElectraDecoders/Private/Utils/AOMedia/ElectraUtilsAV1Video.cpp
@@ -42,7 +42,7 @@ namespace ElectraDecodersUtil
 
 		FString FAV1CodecConfigurationRecord::GetFormatInfo()
 		{
-			return FString::Printf(TEXT("AV1 %s"), *GetFormatInfo(Profile, Level));
+			return FString::Printf(TEXT("AV1, %s"), *GetFormatInfo(Profile, Level));
 		}
 
 
diff --git a/Engine/Plugins/Media/ElectraCodecs/Source/ElectraDecoders/Private/Utils/ElectraCodecFormatUtils.cpp b/Engine/Plugins/Media/ElectraCodecs/Source/ElectraDecoders/Private/Utils/ElectraCodecFormatUtils.cpp
new file mode 100644
--- /dev/null
+++ b/Engine/Plugins/Media/ElectraCodecs/Source/ElectraDecoders/Private/Utils/ElectraCodecFormatUtils.cpp
@@ -0,0 +1,1578 @@
+// Copyright Epic Games, Inc. All Rights Reserved.
+
+#include "ElectraCodecFormatUtils.h"
+#include "Features/IModularFeatures.h"
+#include "Logging/LogMacros.h"
+
+#include "Utils/MPEG/ElectraUtilsMPEGVideo_H264.h"
+#include "Utils/MPEG/ElectraUtilsMPEGVideo_H265.h"
+#include "Utils/MPEG/ElectraUtilsMPEGAudio.h"
+#include "Utils/Google/ElectraUtilsVPxVideo.h"
+#include "Utils/AOMedia/ElectraUtilsAV1Video.h"
+
+#include "Utilities/MP4Boxes/MP4Boxes.h"
+#include "Utilities/UtilitiesMP4.h"
+
+DECLARE_LOG_CATEGORY_EXTERN(LogElectraFormatUtils, Log, All);
+DEFINE_LOG_CATEGORY(LogElectraFormatUtils);
+
+
+FElectraCodecFormatUtilsModularFeature* FElectraCodecFormatUtilsModularFeature::Self = nullptr;
+
+void FElectraCodecFormatUtilsModularFeature::Startup()
+{
+	Self = new FElectraCodecFormatUtilsModularFeature;
+	IModularFeatures::Get().RegisterModularFeature(GetModularFeatureName(), Self);
+}
+
+void FElectraCodecFormatUtilsModularFeature::Shutdown()
+{
+	IModularFeatures::Get().UnregisterModularFeature(GetModularFeatureName(), Self);
+	delete Self;
+	Self = nullptr;
+}
+
+
+namespace
+{
+static const FString MimeType_Video(TEXT("video"));
+static const FString MimeType_Audio(TEXT("audio"));
+static const FString MimeType_Text(TEXT("text"));
+static const FString MimeType_Application(TEXT("application"));
+static const FString MimeType_CodecsPrefix(TEXT("codecs="));
+
+static const FString MimeSubType_TTMLXML(TEXT("ttml+xml"));
+static const FString MimeSubType_MP4(TEXT("mp4"));
+static const FString MimeSubType_VTT(TEXT("vtt"));
+
+static const FString MimeSubType_AVC(TEXT("avc"));
+static const FString MimeSubType_HEVC(TEXT("hevc"));
+static const FString MimeSubType_H264(TEXT("H264"));
+static const FString MimeSubType_H265(TEXT("H265"));
+static const FString MimeSubType_VP8(TEXT("VP8"));
+static const FString MimeSubType_VP9(TEXT("VP9"));
+static const FString MimeSubType_AV1(TEXT("AV1"));
+
+static const FString MimeSubType_MPA(TEXT("MPA"));
+static const FString MimeSubType_MPEG(TEXT("mpeg"));
+static const FString MimeSubType_MP4A(TEXT("mp4a"));
+static const FString MimeSubType_MP3(TEXT("mp3"));
+
+static const FString CodecPrefix_AVC(TEXT("avc"));
+static const FString CodecPrefix_HVC1(TEXT("hvc1"));
+static const FString CodecPrefix_HEV1(TEXT("hev1"));
+static const FString CodecPrefix_VP8(TEXT("vp8"));
+static const FString CodecPrefix_VP08(TEXT("vp08"));
+static const FString CodecPrefix_VP9(TEXT("vp9"));
+static const FString CodecPrefix_VP09(TEXT("vp09"));
+static const FString CodecPrefix_AV1(TEXT("av01"));
+
+static const FString CodecPrefix_MP4A(TEXT("mp4a"));
+static const FString CodecPrefix_Opus(TEXT("opus"));
+static const FString CodecPrefix_Flac(TEXT("flac"));
+
+
+static inline uint32 BitReverse32(uint32 InValue)
+{
+	uint32 rev = 0;
+	for (int32 i = 0; i < 32; ++i)
+	{
+		rev = (rev << 1) | (InValue & 1);
+		InValue >>= 1;
+	}
+	return rev;
+}
+
+static inline constexpr uint32 Make4CC(const uint8 A, const uint8 B, const uint8 C, const uint8 D)
+{
+	return (static_cast<uint32>(A) << 24) | (static_cast<uint32>(B) << 16) | (static_cast<uint32>(C) << 8) | static_cast<uint32>(D);
+}
+
+static FString MakePrintable4CC(uint32 InFourCC)
+{
+	TCHAR tc[4];
+	tc[0] = (TCHAR) ((InFourCC >> 24) & 255);
+	tc[1] = (TCHAR) ((InFourCC >> 16) & 255);
+	tc[2] = (TCHAR) ((InFourCC >>  8) & 255);
+	tc[3] = (TCHAR) ((InFourCC >>  0) & 255);
+	for(int32 i=0;i<4; ++i)
+	{
+		tc[i] = tc[i] >= 32 && tc[i] <= 127 ? tc[i] : TCHAR(' ');
+	}
+	return FString::ConstructFromPtrSize(tc, 4);
+}
+
+static inline void LexFromStringHex(int32& OutValue, const TCHAR* Buffer)
+{
+	OutValue = FCString::Strtoi(Buffer, nullptr, 16);
+}
+
+static inline void LexFromStringHexU64(uint64& OutValue, const TCHAR* Buffer)
+{
+	OutValue = FCString::Strtoui64(Buffer, nullptr, 16);
+}
+
+
+static const TCHAR* GetTypeName(Electra::FCodecTypeFormat::EType InType)
+{
+	switch(InType)
+	{
+		case Electra::FCodecTypeFormat::EType::Video:
+		{
+			return TEXT("video");
+		}
+		case Electra::FCodecTypeFormat::EType::Audio:
+		{
+			return TEXT("audio");
+		}
+		case Electra::FCodecTypeFormat::EType::Subtitle:
+		{
+			return TEXT("subtitle");
+		}
+		case Electra::FCodecTypeFormat::EType::Timecode:
+		{
+			return TEXT("timecode");
+		}
+	}
+	return TEXT("undefined");
+}
+
+
+static void RemapLegacyFourCC(Electra::FCodecTypeFormat& InOut)
+{
+	if (InOut.FourCC == Make4CC('.','m','p','3'))
+	{
+		InOut.FourCC = Make4CC('m','p','g','a');
+	}
+}
+
+/**
+ * For trivial 4CC's that do not require codec parameters, set the default mime type is none is set yet.
+ */
+static void SetDefaultMimeTypeAndRFC6381For4CC(Electra::FCodecTypeFormat& InOut)
+{
+	if (InOut.MimeType.IsEmpty())
+	{
+		switch(InOut.FourCC)
+		{
+			case Make4CC('m','p','g','a'):
+			{
+				InOut.MimeType = TEXT("audio/mpeg");
+				// Unconditionally set the codec rfc to the type we can handle.
+				InOut.RFC6381 = TEXT("mp4a.40.34");
+				break;
+			}
+			case Make4CC('O','p','u','s'):
+			{
+				InOut.MimeType = TEXT("audio/opus");
+				if (InOut.RFC6381.IsEmpty())
+				{
+					InOut.RFC6381 = TEXT("opus");
+				}
+				break;
+			}
+			case Make4CC('f','L','a','C'):
+			{
+				InOut.MimeType = TEXT("audio/flac");
+				if (InOut.RFC6381.IsEmpty())
+				{
+					InOut.RFC6381 = TEXT("flac");
+				}
+				break;
+			}
+			default:
+			{
+				if (InOut.RFC6381.IsEmpty())
+				{
+					InOut.RFC6381 = MakePrintable4CC(InOut.FourCC);
+				}
+				break;
+			}
+		}
+	}
+}
+
+
+
+/**
+ * Attempts to parse the given MIME type string into its components.
+ */
+static bool ParseMimeType(FString& OutType, FString& OutSubType, TArray<FString>& OutCodecs, const FString& InMimeType)
+{
+	if (InMimeType.IsEmpty())
+	{
+		return false;
+	}
+	int32 DelimiterPos = INDEX_NONE;
+	FString TypeSubType = InMimeType.FindChar(TCHAR(';'), DelimiterPos) ? InMimeType.Mid(0, DelimiterPos) : InMimeType;
+	FString Parameter = DelimiterPos != INDEX_NONE ? InMimeType.Mid(DelimiterPos+1) : FString();
+	TypeSubType.TrimStartAndEndInline();
+	if (!TypeSubType.FindChar(TCHAR('/'), DelimiterPos))
+	{
+		return false;
+	}
+	OutType = TypeSubType.Mid(0, DelimiterPos);
+	OutSubType = TypeSubType.Mid(DelimiterPos+1);
+	Parameter.TrimStartAndEndInline();
+	// Check if the parameter is 'codecs'.
+	// Note: We do currently NOT handle 'codecs*' indicating the use of charset, language and escaped characters.
+	if ((DelimiterPos = Parameter.Find(MimeType_CodecsPrefix)) != INDEX_NONE)
+	{
+		FString Codecs = Parameter.Mid(DelimiterPos + MimeType_CodecsPrefix.Len());
+		// Codecs may be quoted. If it is a list it needs to because of the use of the comma ',' separator.
+		Codecs.TrimQuotesInline();
+		// It may be a comma separated list.
+		OutCodecs.Empty();
+		Codecs.ParseIntoArray(OutCodecs, TEXT(","), true);
+		for(int32 i=0; i<OutCodecs.Num(); ++i)
+		{
+			OutCodecs[i].TrimStartAndEndInline();
+			if (OutCodecs[i].IsEmpty())
+			{
+				OutCodecs.RemoveAt(i);
+				--i;
+			}
+		}
+	}
+	return true;
+}
+
+/**
+ * Attempts to determine the type of codec (video, audio, etc.) if the mimetype is provided.
+ */
+static void DetermineCodecTypeFromMimeType(Electra::FCodecTypeFormat& InOut)
+{
+	Electra::FCodecTypeFormat::EType Type = Electra::FCodecTypeFormat::EType::Invalid;
+	// Is a MIME type given?
+	FString MimeType, MimeSubtype;
+	TArray<FString> MimeCodecs;
+	if (ParseMimeType(MimeType, MimeSubtype, MimeCodecs, InOut.MimeType))
+	{
+		if (MimeType.Equals(MimeType_Video, ESearchCase::IgnoreCase))
+		{
+			Type = Electra::FCodecTypeFormat::EType::Video;
+			// Check for audio subtypes that commonly appear in error with the video type
+			if (MimeSubtype.Equals(MimeSubType_MPA, ESearchCase::IgnoreCase) ||
+				MimeSubtype.Equals(MimeSubType_MP4A, ESearchCase::IgnoreCase) ||
+				MimeSubtype.Equals(MimeSubType_MP3, ESearchCase::IgnoreCase))
+			{
+				Type = Electra::FCodecTypeFormat::EType::Audio;
+			}
+		}
+		else if (MimeType.Equals(MimeType_Audio, ESearchCase::IgnoreCase))
+		{
+			Type = Electra::FCodecTypeFormat::EType::Audio;
+		}
+		else if (MimeType.Equals(MimeType_Text, ESearchCase::IgnoreCase))
+		{
+			if (MimeSubtype.Equals(MimeSubType_VTT, ESearchCase::IgnoreCase))
+			{
+				Type = Electra::FCodecTypeFormat::EType::Subtitle;
+			}
+		}
+		else if (MimeType.Equals(MimeType_Application, ESearchCase::IgnoreCase))
+		{
+			if (MimeSubtype.Equals(MimeSubType_TTMLXML, ESearchCase::IgnoreCase) ||
+				MimeSubtype.Equals(MimeSubType_MP4, ESearchCase::IgnoreCase))
+			{
+				Type = Electra::FCodecTypeFormat::EType::Subtitle;
+			}
+		}
+		// Check if the mime type turns up a different format (if any) than the user had already set (if any).
+		if (InOut.Type != Electra::FCodecTypeFormat::EType::Invalid && Type != Electra::FCodecTypeFormat::EType::Invalid && InOut.Type != Type)
+		{
+			UE_LOG(LogElectraFormatUtils, Log, TEXT("Format was set as %s but identifed as %s from MIME type \"%s\""), GetTypeName(InOut.Type), GetTypeName(Type), *InOut.MimeType);
+		}
+
+		// Set the type from the mime type unless the user has explicitly set the type already.
+		if (InOut.Type == Electra::FCodecTypeFormat::EType::Invalid)
+		{
+			InOut.Type = Type;
+		}
+
+		// Is there also a codec specified?
+		if (!MimeCodecs.IsEmpty())
+		{
+			// There must only be a single codec given here, since the structure contains parameters for just one codec.
+			// The user must not pass multiple codecs that may appear in the server received mimetype or conveyed in a playlist.
+			// NOTE: Yes, there could be multiple codecs identifying eg. SVC or MVC on top of the base AVC.
+			//       We do not handle this either.
+			if (MimeCodecs.Num() == 1)
+			{
+				// The codecs parameter is rarely given on a mime type and if it is, if the user has already set an RFC6381 string
+				// we do not look at the mime type provided codec at all. Otherwise we set it as that.
+				if (InOut.RFC6381.IsEmpty())
+				{
+					InOut.RFC6381 = MimeCodecs[0];
+				}
+			}
+			else
+			{
+				UE_LOG(LogElectraFormatUtils, Log, TEXT("The MIME type \"%s\" specifies more than one codec. Ignoring all of them."), *InOut.MimeType);
+			}
+		}
+	}
+}
+
+static bool ParseCodecRFC_AVC(Electra::FCodecTypeFormat& InOut)
+{
+	// avc1 and avc3 (inband SPS/PPS) are recognized.
+	const FString& rfc(InOut.RFC6381);
+	if (rfc.Len() < 4)
+	{
+		return false;
+	}
+	// We support avc1 or avc3 only on the codec since we need the GOP to start with a SAP type 1 or 2.
+	if (rfc[3] != TCHAR('1') && rfc[3] != TCHAR('3'))
+	{
+		return false;
+	}
+	InOut.Type = Electra::FCodecTypeFormat::EType::Video;
+	uint32 Codec4CC = rfc[3] == TCHAR('1') ? Make4CC('a','v','c','1') : Make4CC('a','v','c','3');
+	if (InOut.MimeType.IsEmpty())
+	{
+		InOut.MimeType = TEXT("video/H264");
+	}
+
+	// Profile and level follow?
+	if (rfc.Len() > 5 && rfc[4] == TCHAR('.'))
+	{
+		Electra::FCodecTypeFormat::FVideo::FProfile pl;
+		FString Temp;
+		int32 TempValue;
+		int32 DotPos;
+		rfc.FindLastChar(TCHAR('.'), DotPos);
+		// We recognize the expected format avcC.xxyyzz and for legacy reasons also avcC.xxx.zz
+		if (rfc.Len() == 11 && DotPos == 4)
+		{
+			Temp = rfc.Mid(5, 2);
+			LexFromStringHex(TempValue, *Temp);
+			pl.Profile = TempValue;
+			Temp = rfc.Mid(7, 2);
+			LexFromStringHex(TempValue, *Temp);
+			pl.Constraints = TempValue;
+			Temp = rfc.Mid(9, 2);
+			LexFromStringHex(TempValue, *Temp);
+			pl.Level = TempValue;
+		}
+		else if (DotPos != INDEX_NONE)
+		{
+			Temp = rfc.Mid(5, DotPos-5);
+			LexFromString(TempValue, *Temp);
+			pl.Profile = TempValue;
+			Temp = rfc.Mid(DotPos+1);
+			LexFromString(TempValue, *Temp);
+			pl.Level = TempValue;
+		}
+		else
+		{
+			return false;
+		}
+
+		if (InOut.FourCC && InOut.FourCC != Codec4CC)
+		{
+			UE_LOG(LogElectraFormatUtils, Log, TEXT("FourCC was set as %08x but identifed as %08x from RFC6381 \"%s\""), InOut.FourCC, Codec4CC, *InOut.RFC6381);
+		}
+		// Update the info only if no explicit or different 4CC is set.
+		if (InOut.FourCC == 0 || InOut.FourCC == Codec4CC)
+		{
+			InOut.FourCC = Codec4CC;
+			if (!InOut.Properties.IsType<Electra::FCodecTypeFormat::FVideo>())
+			{
+				InOut.Properties.Emplace<Electra::FCodecTypeFormat::FVideo>();
+			}
+			Electra::FCodecTypeFormat::FVideo& vid(InOut.Properties.Get<Electra::FCodecTypeFormat::FVideo>());
+			vid.Profile = pl;
+		}
+		InOut.HumanReadableFormatInfo = FString::Printf(TEXT("AVC, %s"), *ElectraDecodersUtil::MPEG::H264::FSequenceParameterSet::GetFormatInfo(pl.Profile, pl.Level, (uint32)pl.Constraints));
+	}
+	else
+	{
+		InOut.HumanReadableFormatInfo = TEXT("AVC");
+	}
+	return true;
+}
+
+static bool ParseCodecRFC_HEVC(Electra::FCodecTypeFormat& InOut)
+{
+	// hev1 and hvc1 (inband VPS/SPS/PPS) are recognized.
+	FString rfc(InOut.RFC6381);
+	if (rfc.Len() < 4)
+	{
+		return false;
+	}
+	InOut.Type = Electra::FCodecTypeFormat::EType::Video;
+	uint32 Codec4CC = rfc.StartsWith(TEXT("hev")) ? Make4CC('h','e','v','1') : Make4CC('h','v','c','1');
+	if (InOut.MimeType.IsEmpty())
+	{
+		InOut.MimeType = TEXT("video/H265");
+	}
+
+	int32 DotPos;
+	if (rfc.FindChar(TCHAR('.'), DotPos))
+	{
+		Electra::FCodecTypeFormat::FVideo::FProfile pl;
+		FString Temp;
+		int32 general_profile_space = 0;
+		int32 general_tier_flag = 0;
+		int32 general_profile_idc = 0;
+		int32 general_level_idc = 0;
+		uint32 general_profile_compatibility_flag = 0;
+		uint64 contraint_flags = 0;
+		rfc.RightChopInline(DotPos + 1);
+		// optional general_profile_space
+		if (rfc[0] == TCHAR('A') || rfc[0] == TCHAR('B') || rfc[0] == TCHAR('C'))
+		{
+			general_profile_space = rfc[0] - TCHAR('A') + 1;
+			rfc.RightChopInline(1);
+		}
+		else if (rfc[0] == TCHAR('a') || rfc[0] == TCHAR('b') || rfc[0] == TCHAR('c'))
+		{
+			general_profile_space = rfc[0] - TCHAR('a') + 1;
+			rfc.RightChopInline(1);
+		}
+		// general_profile_idc
+		if (rfc.FindChar(TCHAR('.'), DotPos))
+		{
+			Temp = rfc.Left(DotPos);
+			rfc.RightChopInline(DotPos + 1);
+			LexFromString(general_profile_idc, *Temp);
+		}
+		// general_profile_compatibility_flags
+		if (rfc.FindChar(TCHAR('.'), DotPos))
+		{
+			Temp = rfc.Left(DotPos);
+			rfc.RightChopInline(DotPos + 1);
+			LexFromString(general_profile_compatibility_flag, *Temp);
+		}
+		// general_tier_flag
+		if (rfc[0] != TCHAR('L') && rfc[0] != TCHAR('H') && rfc[0] != TCHAR('l') && rfc[0] != TCHAR('h'))
+		{
+			return false;
+		}
+		else if (rfc[0] == TCHAR('H') || rfc[0] == TCHAR('h'))
+		{
+			general_tier_flag = 1;
+		}
+		rfc.RightChopInline(1);
+		// constraint_flags
+		FString ConstraintFlags;
+		if (rfc.FindChar(TCHAR('.'), DotPos))
+		{
+			ConstraintFlags = rfc.Mid(DotPos + 1);
+			rfc.LeftInline(DotPos);
+			ConstraintFlags.ReplaceInline(TEXT("."), TEXT(""));
+			ConstraintFlags += TEXT("000000000000");
+			ConstraintFlags.LeftInline(12);
+			LexFromStringHexU64(contraint_flags, *ConstraintFlags);
+		}
+		// general_level_idc
+		LexFromString(general_level_idc, *rfc);
+
+		pl.Profile = general_profile_idc;
+		pl.Level = general_level_idc;
+		pl.ProfileSpace = general_profile_space;
+		pl.CompatibilityFlags = BitReverse32(general_profile_compatibility_flag);
+		pl.Tier = general_tier_flag;
+		pl.Constraints = contraint_flags;
+
+		if (InOut.FourCC && InOut.FourCC != Codec4CC)
+		{
+			UE_LOG(LogElectraFormatUtils, Log, TEXT("FourCC was set as %08x but identifed as %08x from RFC6381 \"%s\""), InOut.FourCC, Codec4CC, *InOut.RFC6381);
+		}
+		// Update the info only if no explicit or different 4CC is set.
+		if (InOut.FourCC == 0 || InOut.FourCC == Codec4CC)
+		{
+			InOut.FourCC = Codec4CC;
+			if (!InOut.Properties.IsType<Electra::FCodecTypeFormat::FVideo>())
+			{
+				InOut.Properties.Emplace<Electra::FCodecTypeFormat::FVideo>();
+			}
+			Electra::FCodecTypeFormat::FVideo& vid(InOut.Properties.Get<Electra::FCodecTypeFormat::FVideo>());
+			vid.Profile = pl;
+		}
+		InOut.HumanReadableFormatInfo = FString::Printf(TEXT("HEVC, %s"), *ElectraDecodersUtil::MPEG::H265::FSequenceParameterSet::GetFormatInfo(pl.Profile, pl.Level, pl.Constraints));
+	}
+	else
+	{
+		InOut.HumanReadableFormatInfo = TEXT("HEVC");
+	}
+	return true;
+}
+
+static bool ParseCodecRFC_VPx(Electra::FCodecTypeFormat& InOut)
+{
+	FString rfc(InOut.RFC6381);
+	if (rfc.Len() < 4)
+	{
+		return false;
+	}
+	InOut.Type = Electra::FCodecTypeFormat::EType::Video;
+	uint32 Codec4CC = (rfc.StartsWith(TEXT("vp08")) || rfc.StartsWith(TEXT("vp8"))) ? Make4CC('v','p','0','8') : Make4CC('v','p','0','9');
+	if (InOut.FourCC && InOut.FourCC != Codec4CC)
+	{
+		UE_LOG(LogElectraFormatUtils, Log, TEXT("FourCC was set as %08x but identifed as %08x from RFC6381 \"%s\""), InOut.FourCC, Codec4CC, *InOut.RFC6381);
+		InOut.FourCC = Codec4CC;
+	}
+	if (InOut.MimeType.IsEmpty())
+	{
+		InOut.MimeType = Codec4CC == Make4CC('v','p','0','8') ?	TEXT("video/VP8") : TEXT("video/VP9");
+	}
+
+	if (rfc.Len() > 4)
+	{
+		Electra::FCodecTypeFormat::FVideo::FProfile pl;
+		FString Temp;
+		int32 DotPos;
+		if (rfc.FindChar(TCHAR('.'), DotPos))
+		{
+			int32 Components[8] {0}, NumComponents=0;
+			rfc.RightChopInline(DotPos + 1);
+			while(rfc.Len() && NumComponents<UE_ARRAY_COUNT(Components))
+			{
+				if (rfc.FindChar(TCHAR('.'), DotPos))
+				{
+					Temp = rfc.Left(DotPos);
+					rfc.RightChopInline(DotPos + 1);
+					LexFromString(Components[NumComponents++], *Temp);
+				}
+				else
+				{
+					LexFromString(Components[NumComponents++], *rfc);
+					break;
+				}
+			}
+			// If a parameter string is given then the first three components are mandatory.
+			if (NumComponents < 3)
+			{
+				return false;
+			}
+
+			if (!InOut.Properties.IsType<Electra::FCodecTypeFormat::FVideo>())
+			{
+				InOut.Properties.Emplace<Electra::FCodecTypeFormat::FVideo>();
+			}
+			Electra::FCodecTypeFormat::FVideo& vid(InOut.Properties.Get<Electra::FCodecTypeFormat::FVideo>());
+			pl.Profile = Components[0];
+			pl.Level = Components[1];
+			vid.Profile = pl;
+			vid.BitDepth = Components[2];
+			Electra::FCodecTypeFormat::FVideo::FColorInfo ci;
+			ci.chromaSubsampling = 1;
+			ci.colourPrimaries = 1;
+			ci.transferCharacteristics = 1;
+			ci.matrixCoefficients = 1;
+			ci.videoFullRangeFlag = 0;
+			if (NumComponents > 3)
+			{
+				ci.chromaSubsampling = (uint8) Components[3];
+				ci.colourPrimaries = (uint8) Components[4];
+				ci.transferCharacteristics  = (uint8) Components[5];
+				ci.matrixCoefficients = (uint8) Components[6];
+				ci.videoFullRangeFlag = (uint8) Components[7];
+			}
+			vid.OptColorInfo = ci;
+			InOut.HumanReadableFormatInfo = FString::Printf(TEXT("VP%d, Profile %u, level %u.%u"), Codec4CC == Make4CC('v','p','0','8') ? 8 : 9, pl.Profile, pl.Level / 10, pl.Level % 10);
+		}
+	}
+	else
+	{
+		InOut.HumanReadableFormatInfo = FString::Printf(TEXT("VP%d"), Codec4CC == Make4CC('v','p','0','8') ? 8 : 9);
+	}
+	return true;
+}
+
+static bool ParseCodecRFC_AV1(Electra::FCodecTypeFormat& InOut)
+{
+	const FString& rfc(InOut.RFC6381);
+	if (rfc.Len() < 4)
+	{
+		return false;
+	}
+	InOut.Type = Electra::FCodecTypeFormat::EType::Video;
+	if (InOut.MimeType.IsEmpty())
+	{
+		InOut.MimeType = TEXT("video/AV1");
+	}
+	const uint32 Codec4CC = Make4CC('a','v','0','1');
+	if (InOut.FourCC && InOut.FourCC != Codec4CC)
+	{
+		UE_LOG(LogElectraFormatUtils, Log, TEXT("FourCC was set as %08x but identifed as %08x from RFC6381 \"%s\""), InOut.FourCC, Codec4CC, *InOut.RFC6381);
+		InOut.FourCC = Codec4CC;
+	}
+	if (rfc.Len() > 4)
+	{
+		TArray<FString> Parts;
+		rfc.ParseIntoArray(Parts, TEXT("."), false);
+		// Profile, level, tier and bitdepth are required.
+		if (Parts.Num() < 4)
+		{
+			return false;
+		}
+		Electra::FCodecTypeFormat::FVideo::FProfile pl;
+		ElectraDecodersUtil::AV1Video::FAV1CodecConfigurationRecord::FColorInfo av1Col;
+		int32 NumBitsLuma = 8;
+		auto ParseParts = [&]() -> bool
+		{
+			for(int32 i=1; i<Parts.Num(); ++i)
+			{
+				switch(i)
+				{
+					case 1:
+					{
+						// Profile. Has to be a single digit.
+						if (Parts[i].Len() > 1)
+						{
+							return false;
+						}
+						LexFromString(pl.Profile, *Parts[i]);
+						break;
+					}
+					case 2:
+					{
+						// Level and tier. Needs to be 2 digits and an 'M' or 'H'.
+						if (Parts[i].Len() != 3 || (Parts[i][2]!=TCHAR('M') && Parts[i][2]!=TCHAR('m') && Parts[i][2]!=TCHAR('H') && Parts[i][2]!=TCHAR('h')))
+						{
+							return false;
+						}
+						pl.Tier = Parts[i][2]==TCHAR('H') || Parts[i][2]==TCHAR('h') ? 1 : 0;
+						Parts[i].LeftChopInline(1);
+						LexFromString(pl.Level, *Parts[i]);
+						break;
+					}
+					case 3:
+					{
+						// Bitdepth, 2 digits
+						if (Parts[i].Len() != 2)
+						{
+							return false;
+						}
+						LexFromString(NumBitsLuma, *Parts[i]);
+						break;
+					}
+					/*
+						From this point on missing or malformed components are not critical. As per the standard:
+							"If any character that is not '.', digits, part of the AV1 4CC, or a tier value is encountered,
+							the string SHALL be interpreted ignoring all the characters starting from that character."
+					*/
+					case 4:
+					{
+						// Monochrome, 1 digit
+						if (Parts[i].Len() != 1)
+						{
+							return true;
+						}
+						LexFromString(av1Col.monochrome, *Parts[i]);
+						break;
+					}
+					case 5:
+					{
+						// Chroma subsampling, 3 digits
+						if (Parts[i].Len() != 3)
+						{
+							return true;
+						}
+						if (FChar::IsDigit(Parts[i][0])) av1Col.chromaSubsamplingX = Parts[i][0]-TCHAR('0'); else return true;
+						if (FChar::IsDigit(Parts[i][1])) av1Col.chromaSubsamplingY = Parts[i][1]-TCHAR('0'); else return true;
+						if (FChar::IsDigit(Parts[i][2])) av1Col.chromaSamplingPosition = Parts[i][2]-TCHAR('0'); else return true;
+						break;
+					}
+					case 6:
+					{
+						// colorPrimaries.
+						if (Parts[i].Len() != 2)
+						{
+							return false;
+						}
+						LexFromString(av1Col.colorPrimaries, *Parts[i]);
+						break;
+					}
+					case 7:
+					{
+						// transferCharacteristics.
+						if (Parts[i].Len() != 2)
+						{
+							return false;
+						}
+						LexFromString(av1Col.transferCharacteristics, *Parts[i]);
+						break;
+					}
+					case 8:
+					{
+						// matrixCoefficients.
+						if (Parts[i].Len() != 2)
+						{
+							return false;
+						}
+						LexFromString(av1Col.matrixCoefficients, *Parts[i]);
+						break;
+					}
+					case 9:
+					{
+						// videoFullRangeFlag.
+						if (Parts[i].Len() != 1)
+						{
+							return false;
+						}
+						LexFromString(av1Col.videoFullRangeFlag, *Parts[i]);
+						break;
+					}
+					default:
+					{
+						break;
+					}
+				}
+			}
+			return true;
+		};
+		if (!ParseParts())
+		{
+			return false;
+		}
+
+		if (!InOut.Properties.IsType<Electra::FCodecTypeFormat::FVideo>())
+		{
+			InOut.Properties.Emplace<Electra::FCodecTypeFormat::FVideo>();
+		}
+		Electra::FCodecTypeFormat::FVideo& vid(InOut.Properties.Get<Electra::FCodecTypeFormat::FVideo>());
+		vid.Profile = pl;
+		vid.BitDepth = NumBitsLuma;
+		Electra::FCodecTypeFormat::FVideo::FColorInfo ci;
+		ci.colourPrimaries = av1Col.colorPrimaries;
+		ci.transferCharacteristics = av1Col.transferCharacteristics;
+		ci.matrixCoefficients = av1Col.matrixCoefficients;
+		ci.videoFullRangeFlag = av1Col.videoFullRangeFlag;
+		Electra::FCodecTypeFormat::FVideo::FColorInfo::FChromaInfo cri;
+		cri.monochrome = av1Col.monochrome;
+		cri.chromaSubsamplingX = av1Col.chromaSubsamplingX;
+		cri.chromaSubsamplingY = av1Col.chromaSubsamplingY;
+		cri.chromaSamplingPosition = av1Col.chromaSamplingPosition;
+		ci.ChromaInfo = cri;
+		vid.OptColorInfo = ci;
+		InOut.HumanReadableFormatInfo = FString::Printf(TEXT("AV1, %s"), *ElectraDecodersUtil::AV1Video::FAV1CodecConfigurationRecord::GetFormatInfo(pl.Profile, pl.Level));
+	}
+	else
+	{
+		InOut.HumanReadableFormatInfo = TEXT("AV1");
+	}
+	return true;
+}
+
+
+
+static void MakeHumanReadableRFC_MP4A(Electra::FCodecTypeFormat& InOut, int32 InAudioObjectType, int32 InMpegVersion, int32 InLayer)
+{
+	if (InOut.FourCC == Make4CC('m','p','4','a'))
+	{
+		if (InAudioObjectType)
+		{
+			InOut.HumanReadableFormatInfo = FString::Printf(TEXT("AAC (%s)"), InAudioObjectType==2 ? TEXT("LC") : InAudioObjectType==5 ? TEXT("SBR") : InAudioObjectType==29 ? TEXT("PS") : TEXT("unknown"));
+		}
+		else
+		{
+			InOut.HumanReadableFormatInfo = TEXT("AAC");
+		}
+	}
+	else
+	{
+		FString mp(InMpegVersion ? FString::Printf(TEXT("MPEG%d"), InMpegVersion) : TEXT("MPEG"));
+		if (InLayer)
+		{
+			InOut.HumanReadableFormatInfo = FString::Printf(TEXT("%s layer %d"), *mp, InLayer);
+		}
+		else
+		{
+			InOut.HumanReadableFormatInfo = FString::Printf(TEXT("%s audio"), *mp);
+		}
+	}
+}
+
+static bool ParseCodecRFC_MP4A(Electra::FCodecTypeFormat& InOut)
+{
+	const FString& rfc(InOut.RFC6381);
+	if (rfc.Len() < 4)
+	{
+		return false;
+	}
+	InOut.Type = Electra::FCodecTypeFormat::EType::Audio;
+	TArray<FString> Parts;
+	rfc.ParseIntoArray(Parts, TEXT("."));
+	uint32 Codec4CC = 0;
+	int32 AudioObjectType = 0;
+	int32 Layer = 0;
+	int32 Mpeg = 0;
+	if (Parts.Num() > 1)
+	{
+		int32 OTI = 0;
+		LexFromStringHex(OTI, *Parts[1]);
+		switch(static_cast<ElectraDecodersUtil::MPEG::FESDescriptor::FObjectTypeID>(OTI))
+		{
+			// ISO/IEC 14496-3
+			case ElectraDecodersUtil::MPEG::FESDescriptor::FObjectTypeID::MPEG4_Audio:
+			{
+				if (Parts.Num() < 3)
+				{
+					UE_LOG(LogElectraFormatUtils, Log, TEXT("RFC6381 \"%s\" type needs the audio type indicator"), *InOut.RFC6381);
+					return false;
+				}
+				LexFromString(AudioObjectType, *Parts[2]);
+				switch(AudioObjectType)
+				{
+					case 2:		// AAC-LC
+					case 5:		// SBR
+					case 29:	// PS
+					{
+						Codec4CC = Make4CC('m','p','4','a');
+						if (InOut.MimeType.IsEmpty())
+						{
+							InOut.MimeType = TEXT("audio/mp4");
+						}
+						break;
+					}
+					case 32:	// Layer-1
+					case 33:	// Layer-2
+					case 34:	// Layer-3
+					{
+						Codec4CC = Make4CC('m','p','g','a');
+						if (InOut.MimeType.IsEmpty())
+						{
+							InOut.MimeType = TEXT("audio/mpeg");
+						}
+						Layer = AudioObjectType - 31;
+						Mpeg = 1;
+						break;
+					}
+					default:
+					{
+						// not supported
+						return false;
+					}
+				}
+				break;
+			}
+			// ISO/IEC 13818-7
+			case ElectraDecodersUtil::MPEG::FESDescriptor::FObjectTypeID::MPEG2_AAC_LC:
+			{
+				Codec4CC = Make4CC('m','p','4','a');
+				AudioObjectType = 2;
+				if (InOut.MimeType.IsEmpty())
+				{
+					InOut.MimeType = TEXT("audio/mp4a");
+				}
+				break;
+			}
+			// ISO/IEC 13818-3
+			case ElectraDecodersUtil::MPEG::FESDescriptor::FObjectTypeID::MPEG2_Audio:
+			// ISO/IEC 11172-3
+			case ElectraDecodersUtil::MPEG::FESDescriptor::FObjectTypeID::MPEG1_Audio:
+			{
+				Mpeg = static_cast<ElectraDecodersUtil::MPEG::FESDescriptor::FObjectTypeID>(OTI) == ElectraDecodersUtil::MPEG::FESDescriptor::FObjectTypeID::MPEG1_Audio ? 1 : 2;
+				// Both can be handled by an MPEG1 layer III decoder
+				Codec4CC = Make4CC('m','p','g','a');
+				if (InOut.MimeType.IsEmpty())
+				{
+					InOut.MimeType = TEXT("audio/mpeg");
+				}
+				break;
+			}
+			default:
+			{
+				return false;
+			}
+		}
+	}
+
+	if (InOut.FourCC && Codec4CC && InOut.FourCC != Codec4CC)
+	{
+		UE_LOG(LogElectraFormatUtils, Log, TEXT("FourCC was set as %08x but identifed as %08x from RFC6381 \"%s\""), InOut.FourCC, Codec4CC, *InOut.RFC6381);
+		InOut.FourCC = Codec4CC;
+	}
+
+	if (!InOut.Properties.IsType<Electra::FCodecTypeFormat::FAudio>())
+	{
+		InOut.Properties.Emplace<Electra::FCodecTypeFormat::FAudio>();
+	}
+
+	MakeHumanReadableRFC_MP4A(InOut, AudioObjectType, Mpeg, Layer);
+	return true;
+}
+
+
+static bool DetermineCodecTypeFromRFC(Electra::FCodecTypeFormat& InOut)
+{
+	bool bOk = true;
+	// Check the codec RFC
+	if (!InOut.RFC6381.IsEmpty())
+	{
+		if (InOut.RFC6381.StartsWith(CodecPrefix_AVC))
+		{
+			bOk = ParseCodecRFC_AVC(InOut);
+		}
+		else if (InOut.RFC6381.StartsWith(CodecPrefix_HVC1) || InOut.RFC6381.StartsWith(CodecPrefix_HEV1))
+		{
+			bOk = ParseCodecRFC_HEVC(InOut);
+		}
+		else if (InOut.RFC6381.StartsWith(CodecPrefix_VP08) || InOut.RFC6381.StartsWith(CodecPrefix_VP8))
+		{
+			bOk = ParseCodecRFC_VPx(InOut);
+		}
+		else if (InOut.RFC6381.StartsWith(CodecPrefix_VP09) || InOut.RFC6381.StartsWith(CodecPrefix_VP9))
+		{
+			bOk = ParseCodecRFC_VPx(InOut);
+		}
+		else if (InOut.RFC6381.StartsWith(CodecPrefix_AV1))
+		{
+			bOk = ParseCodecRFC_AV1(InOut);
+		}
+
+		else if (InOut.RFC6381.StartsWith(CodecPrefix_MP4A))
+		{
+			bOk = ParseCodecRFC_MP4A(InOut);
+		}
+		else if (InOut.RFC6381.StartsWith(CodecPrefix_Opus))
+		{
+			InOut.Type = Electra::FCodecTypeFormat::EType::Audio;
+			if (!InOut.Properties.IsType<Electra::FCodecTypeFormat::FAudio>())
+			{
+				InOut.Properties.Emplace<Electra::FCodecTypeFormat::FAudio>();
+			}
+			InOut.HumanReadableFormatInfo = TEXT("Opus");
+		}
+		else if (InOut.RFC6381.StartsWith(CodecPrefix_Flac))
+		{
+			InOut.Type = Electra::FCodecTypeFormat::EType::Audio;
+			if (!InOut.Properties.IsType<Electra::FCodecTypeFormat::FAudio>())
+			{
+				InOut.Properties.Emplace<Electra::FCodecTypeFormat::FAudio>();
+			}
+			InOut.HumanReadableFormatInfo = TEXT("Free Lossless Audio Codec (FLAC)");
+		}
+	}
+	return bOk;
+}
+
+
+static int32 DetermineCodecTypeFromDCR_avcc(Electra::FCodecTypeFormat& InOut)
+{
+	const TArray<uint8>* avcC_Box = InOut.ExtraBoxes.Find(Make4CC('a','v','c','C'));
+	if (avcC_Box)
+	{
+		if (InOut.MimeType.IsEmpty())
+		{
+			InOut.MimeType = TEXT("video/H264");
+		}
+		ElectraDecodersUtil::MPEG::H264::FAVCDecoderConfigurationRecord dcr;
+		if (dcr.Parse(*avcC_Box))
+		{
+			InOut.DCR = *avcC_Box;
+			InOut.Type = Electra::FCodecTypeFormat::EType::Video;
+			if (!InOut.Properties.IsType<Electra::FCodecTypeFormat::FVideo>())
+			{
+				InOut.Properties.Emplace<Electra::FCodecTypeFormat::FVideo>();
+			}
+
+			ElectraDecodersUtil::MPEG::H264::FSequenceParameterSet sps0;
+			bool bSps0Ok = false;
+			if (!dcr.GetSequenceParameterSets().IsEmpty())
+			{
+				bSps0Ok = ElectraDecodersUtil::MPEG::H264::ParseSequenceParameterSet(sps0, dcr.GetSequenceParameterSets()[0]);
+			}
+			Electra::FCodecTypeFormat::FVideo& vid(InOut.Properties.Get<Electra::FCodecTypeFormat::FVideo>());
+			if ((vid.Profile.Profile && vid.Profile.Profile != dcr.GetAVCProfileIndication()) ||
+				(vid.Profile.Level && vid.Profile.Level != dcr.GetAVCLevelIndication()))
+			{
+				UE_LOG(LogElectraFormatUtils, Log, TEXT("AVC profile and level was set as %d,%d but specified as %d,%d in AVCDecoderConfigurationRecord"), vid.Profile.Profile, vid.Profile.Level, dcr.GetAVCProfileIndication(), dcr.GetAVCLevelIndication());
+			}
+			vid.Profile.Profile = dcr.GetAVCProfileIndication();
+			vid.Profile.Level = dcr.GetAVCLevelIndication();
+			vid.Profile.Constraints = dcr.GetProfileCompatibility();
+			InOut.CSD = dcr.GetCodecSpecificData();
+			int32 FourCC = InOut.CSD.IsEmpty() ? Make4CC('a','v','c','3') : Make4CC('a','v','c','1');
+			if (InOut.FourCC && InOut.FourCC != FourCC)
+			{
+				UE_LOG(LogElectraFormatUtils, Log, TEXT("FourCC was set as %08x but identifed as %08x from AVCDecoderConfigurationRecord"), InOut.FourCC, FourCC);
+			}
+			InOut.FourCC = FourCC;
+
+			vid.BitDepth = dcr.GetBitDepthLuma();
+			if (bSps0Ok)
+			{
+				int32 w, h;
+				sps0.GetDisplaySize(w, h);
+				// Check dimensions
+				if ((vid.Width && vid.Width != w) ||
+					(vid.Height && vid.Height != h))
+				{
+					UE_LOG(LogElectraFormatUtils, Log, TEXT("Dimensions were set as %u*%u but given as %d*%d by the SPS"), vid.Width, vid.Height, w, h);
+					vid.Width = w > 0 ? (uint32)w : vid.Width;
+					vid.Height = h > 0 ? (uint32)h : vid.Height;
+				}
+				else if ((!vid.Width && w) || (!vid.Height && h))
+				{
+					vid.Width = (uint32)w;
+					vid.Height = (uint32)h;
+				}
+				// Check timing
+				ElectraDecodersUtil::FFractionalValue vuiTiming = sps0.GetTiming();
+				if (vuiTiming.Num && vuiTiming.Denom)
+				{
+					double fps = (double)vuiTiming.Num / vuiTiming.Denom;
+					double cfps = vid.FrameRate.IsValid() ? vid.FrameRate.AsDecimal() : 0.0;
+					double r = cfps / fps;
+					// There is no guarantee that the VUI timing in the SPS is correct if the timing is conveyed in the stream in SEI messages.
+					if ((r < 1.2 && r > 0.8) || (cfps == 0.0 && fps <= 240.0))
+					{
+						vid.FrameRate.Numerator = (int32)vuiTiming.Num;
+						vid.FrameRate.Denominator = (int32)vuiTiming.Denom;
+					}
+				}
+			}
+
+			InOut.RFC6381 = dcr.GetCodecSpecifierRFC6381(FourCC == Make4CC('a','v','c','1') ? TEXT("avc1") : TEXT("avc3"));
+			InOut.HumanReadableFormatInfo = FString::Printf(TEXT("AVC, %s"), *ElectraDecodersUtil::MPEG::H264::FSequenceParameterSet::GetFormatInfo(vid.Profile.Profile, vid.Profile.Level, (uint32)vid.Profile.Constraints));
+			return 1;
+		}
+		return -1;
+	}
+	return 0;
+}
+
+static int32 DetermineCodecTypeFromDCR_hvcc(Electra::FCodecTypeFormat& InOut)
+{
+	const TArray<uint8>* hvcC_Box = InOut.ExtraBoxes.Find(Make4CC('h','v','c','C'));
+	if (hvcC_Box)
+	{
+		if (InOut.MimeType.IsEmpty())
+		{
+			InOut.MimeType = TEXT("video/H265");
+		}
+		ElectraDecodersUtil::MPEG::H265::FHEVCDecoderConfigurationRecord dcr;
+		if (dcr.Parse(*hvcC_Box))
+		{
+			InOut.DCR = *hvcC_Box;
+			InOut.Type = Electra::FCodecTypeFormat::EType::Video;
+			if (!InOut.Properties.IsType<Electra::FCodecTypeFormat::FVideo>())
+			{
+				InOut.Properties.Emplace<Electra::FCodecTypeFormat::FVideo>();
+			}
+
+			ElectraDecodersUtil::MPEG::H265::FSequenceParameterSet sps0;
+			bool bSps0Ok = false;
+			if (!dcr.GetSequenceParameterSets().IsEmpty())
+			{
+				bSps0Ok = ElectraDecodersUtil::MPEG::H265::ParseSequenceParameterSet(sps0, dcr.GetSequenceParameterSets()[0]);
+			}
+			Electra::FCodecTypeFormat::FVideo& vid(InOut.Properties.Get<Electra::FCodecTypeFormat::FVideo>());
+			if ((vid.Profile.Profile && vid.Profile.Profile != dcr.GetGeneralProfileIDC()) ||
+				(vid.Profile.Level && vid.Profile.Level != dcr.GetGeneralLevelIDC()))
+			{
+				UE_LOG(LogElectraFormatUtils, Log, TEXT("HEVC profile and level was set as %d,%d but specified as %d,%d in HEVCDecoderConfigurationRecord"), vid.Profile.Profile, vid.Profile.Level, dcr.GetGeneralProfileIDC(), dcr.GetGeneralLevelIDC());
+			}
+			vid.Profile.Tier = dcr.GetGeneralTierFlag();
+			vid.Profile.ProfileSpace = dcr.GetGeneralProfileSpace();
+			vid.Profile.Profile = dcr.GetGeneralProfileIDC();
+			vid.Profile.Level = dcr.GetGeneralLevelIDC();
+			vid.Profile.Constraints = dcr.GetGeneralConstraintIndicatorFlags();
+			vid.Profile.CompatibilityFlags = dcr.GetGeneralProfileCompatibilityFlags();
+			InOut.CSD = dcr.GetCodecSpecificData();
+			int32 FourCC = InOut.CSD.IsEmpty() ? Make4CC('h','v','c','1') : Make4CC('h','e','v','1');
+			if (InOut.FourCC && InOut.FourCC != FourCC)
+			{
+				UE_LOG(LogElectraFormatUtils, Log, TEXT("FourCC was set as %08x but identifed as %08x from HEVCDecoderConfigurationRecord"), InOut.FourCC, FourCC);
+			}
+			InOut.FourCC = FourCC;
+
+			vid.BitDepth = 8 + dcr.GetBitDepthLumaMinus8();
+			if (bSps0Ok)
+			{
+				uint32 w, h;
+				sps0.GetDisplaySize(w, h);
+				// Check dimensions
+				if ((vid.Width && vid.Width != w) ||
+					(vid.Height && vid.Height != h))
+				{
+					UE_LOG(LogElectraFormatUtils, Log, TEXT("Dimensions were set as %u*%u but given as %d*%d by the SPS"), vid.Width, vid.Height, w, h);
+					vid.Width = w > 0 ? w : vid.Width;
+					vid.Height = h > 0 ? h : vid.Height;
+				}
+				else if ((!vid.Width && w) || (!vid.Height && h))
+				{
+					vid.Width = w;
+					vid.Height = h;
+				}
+				// Check timing
+				ElectraDecodersUtil::FFractionalValue vuiTiming = sps0.GetTiming();
+				if (vuiTiming.Num && vuiTiming.Denom)
+				{
+					double fps = (double)vuiTiming.Num / vuiTiming.Denom;
+					double cfps = vid.FrameRate.IsValid() ? vid.FrameRate.AsDecimal() : 0.0;
+					double r = cfps / fps;
+					// There is no guarantee that the VUI timing in the SPS is correct if the timing is conveyed in the stream in SEI messages.
+					if ((r < 1.2 && r > 0.8) || (cfps == 0.0 && fps <= 240.0))
+					{
+						vid.FrameRate.Numerator = (int32)vuiTiming.Num;
+						vid.FrameRate.Denominator = (int32)vuiTiming.Denom;
+					}
+				}
+			}
+
+			InOut.RFC6381 = dcr.GetCodecSpecifierRFC6381(FourCC == Make4CC('h','e','v','1') ? TEXT("hev1") : TEXT("hvc1"));
+			InOut.HumanReadableFormatInfo = FString::Printf(TEXT("HEVC, %s"), *ElectraDecodersUtil::MPEG::H265::FSequenceParameterSet::GetFormatInfo(vid.Profile.Profile, vid.Profile.Level, (uint32)vid.Profile.Constraints));
+			return 1;
+		}
+		return -1;
+	}
+	return 0;
+}
+
+static int32 DetermineCodecTypeFromDCR_vpcc(Electra::FCodecTypeFormat& InOut)
+{
+	const TArray<uint8>* vpcC_Box = InOut.ExtraBoxes.Find(Make4CC('v','p','c','C'));
+	if (vpcC_Box && vpcC_Box->Num() > 4)
+	{
+		if (InOut.MimeType.IsEmpty())
+		{
+			InOut.MimeType = InOut.FourCC == Make4CC('v','p','0','8') ?	TEXT("video/VP8") : TEXT("video/VP9");
+		}
+		ElectraDecodersUtil::VPxVideo::FVPCodecConfigurationRecord dcr;
+		// The VPCodecConfigurationBox is a full box, so we need to skip the first 4 bytes.
+		auto BoxData = MakeConstArrayView<const uint8>(vpcC_Box->GetData() + 4, vpcC_Box->Num() - 4);
+		if (dcr.Parse(InOut.FourCC == Make4CC('v','p','0','8') ? 8 : 9, BoxData))
+		{
+			InOut.DCR = BoxData;
+			InOut.Type = Electra::FCodecTypeFormat::EType::Video;
+			if (!InOut.Properties.IsType<Electra::FCodecTypeFormat::FVideo>())
+			{
+				InOut.Properties.Emplace<Electra::FCodecTypeFormat::FVideo>();
+			}
+			Electra::FCodecTypeFormat::FVideo& vid(InOut.Properties.Get<Electra::FCodecTypeFormat::FVideo>());
+			if ((vid.Profile.Profile && vid.Profile.Profile != dcr.GetProfile()) ||
+				(vid.Profile.Level && vid.Profile.Level != dcr.GetLevel()))
+			{
+				UE_LOG(LogElectraFormatUtils, Log, TEXT("VPx profile and level was set as %d,%d but specified as %d,%d in VPCodecConfigurationRecord"), vid.Profile.Profile, vid.Profile.Level, dcr.GetProfile(), dcr.GetLevel());
+			}
+			vid.Profile.Profile = dcr.GetProfile();
+			vid.Profile.Level = dcr.GetLevel();
+			InOut.CSD = dcr.GetCodecSpecificData();
+
+			vid.BitDepth = dcr.GetBitDepth();
+			Electra::FCodecTypeFormat::FVideo::FColorInfo ci;
+			const ElectraDecodersUtil::VPxVideo::FVPCodecConfigurationRecord::FColorInfo& vpxCol = dcr.GetColorInfo();
+			ci.chromaSubsampling = vpxCol.chromaSubsampling;
+			ci.colourPrimaries = vpxCol.colourPrimaries;
+			ci.transferCharacteristics = vpxCol.transferCharacteristics;
+			ci.matrixCoefficients = vpxCol.matrixCoefficients;
+			ci.videoFullRangeFlag = vpxCol.videoFullRangeFlag;
+			vid.OptColorInfo = ci;
+
+			FString rfc(FString::Printf(TEXT("vp0%d.%02d.%02d.%02d.%02d.%02d.%02d.%02d.%02d"), InOut.FourCC == Make4CC('v','p','0','8') ? 8 : 9,
+							vid.Profile.Profile, vid.Profile.Level, vid.BitDepth,ci.chromaSubsampling, ci.colourPrimaries, ci.transferCharacteristics, ci.matrixCoefficients, ci.videoFullRangeFlag));
+			if (rfc.EndsWith(TEXT(".01.01.01.01.00")))
+			{
+				rfc.LeftChopInline(15);
+			}
+			InOut.RFC6381 = MoveTemp(rfc);
+			InOut.HumanReadableFormatInfo = FString::Printf(TEXT("VP%d, Profile %u, level %u.%u"), InOut.FourCC == Make4CC('v','p','0','8') ? 8 : 9, vid.Profile.Profile, vid.Profile.Level / 10, vid.Profile.Level % 10);
+			return 1;
+		}
+		return -1;
+	}
+	return 0;
+}
+
+static int32 DetermineCodecTypeFromDCR_av1c(Electra::FCodecTypeFormat& InOut)
+{
+	const TArray<uint8>* av1C_Box = InOut.ExtraBoxes.Find(Make4CC('a','v','1','C'));
+	if (av1C_Box && av1C_Box->Num() > 4)
+	{
+		if (InOut.MimeType.IsEmpty())
+		{
+			InOut.MimeType = TEXT("video/AV1");
+		}
+		ElectraDecodersUtil::AV1Video::FAV1CodecConfigurationRecord dcr;
+		if (dcr.Parse(*av1C_Box))
+		{
+			InOut.DCR = *av1C_Box;
+			InOut.Type = Electra::FCodecTypeFormat::EType::Video;
+			if (!InOut.Properties.IsType<Electra::FCodecTypeFormat::FVideo>())
+			{
+				InOut.Properties.Emplace<Electra::FCodecTypeFormat::FVideo>();
+			}
+			Electra::FCodecTypeFormat::FVideo& vid(InOut.Properties.Get<Electra::FCodecTypeFormat::FVideo>());
+			if ((vid.Profile.Profile && vid.Profile.Profile != dcr.GetProfile()) ||
+				(vid.Profile.Level && vid.Profile.Level != dcr.GetLevel()))
+			{
+				UE_LOG(LogElectraFormatUtils, Log, TEXT("AV1 profile and level was set as %d,%d but specified as %d,%d in AV1CodecConfigurationRecord"), vid.Profile.Profile, vid.Profile.Level, dcr.GetProfile(), dcr.GetLevel());
+			}
+			vid.Profile.Profile = dcr.GetProfile();
+			vid.Profile.Level = dcr.GetLevel();
+			InOut.CSD = dcr.GetCodecSpecificData();
+			if (!InOut.CSD.IsEmpty())
+			{
+				const uint8* CurrentOBU = InOut.CSD.GetData();
+				const uint8* EndOfOBUs = CurrentOBU + InOut.CSD.Num();
+				while(CurrentOBU < EndOfOBUs)
+				{
+					const uint8 OBUheaderByte = *CurrentOBU;
+					if ((OBUheaderByte & 0x80) != 0)
+					{
+						UE_LOG(LogElectraFormatUtils, Log, TEXT("obu_forbidden_bit not zero"));
+						return -1;
+					}
+					const ElectraDecodersUtil::AV1Video::EOBUType OBUtype = (const ElectraDecodersUtil::AV1Video::EOBUType)(OBUheaderByte >> 3);
+					const uint8 obu_extension_flag = (OBUheaderByte >> 2) & 1;
+					const uint8 obu_has_size_field = (OBUheaderByte >> 1) & 1;
+					const uint8 obu_reserved_1bit = OBUheaderByte & 1;
+					CurrentOBU += obu_extension_flag ? 2 : 1;
+					uint32 obu_size;
+					uint32 nr = 0;
+					if (obu_has_size_field)
+					{
+						ElectraDecodersUtil::AV1Video::FBitstreamReaderAV1 bs(CurrentOBU, EndOfOBUs - CurrentOBU);
+						obu_size = bs.leb128(nr);
+						CurrentOBU += nr;
+					}
+					else
+					{
+						obu_size = EndOfOBUs - CurrentOBU;
+					}
+
+					if (OBUtype == ElectraDecodersUtil::AV1Video::EOBUType::OBU_SEQUENCE_HEADER)
+					{
+						ElectraDecodersUtil::AV1Video::FBitstreamReaderAV1 bs(CurrentOBU, obu_size);
+						ElectraDecodersUtil::AV1Video::FSequenceHeader sh;
+						if (!sh.ParseFrom(bs))
+						{
+							UE_LOG(LogElectraFormatUtils, Log, TEXT("Failed to parse sequence_header_obu"));
+							return -1;
+						}
+
+						// Check dimensions
+						uint32 w = 1 + sh.max_frame_width_minus_1;
+						uint32 h = 1 + sh.max_frame_height_minus_1;
+						if ((vid.Width && vid.Width != w) ||
+							(vid.Height && vid.Height != h))
+						{
+							UE_LOG(LogElectraFormatUtils, Log, TEXT("Dimensions were set as %u*%u but given as %d*%d by the sequence header"), vid.Width, vid.Height, w, h);
+							vid.Width = w ? w : vid.Width;
+							vid.Height = h ? h : vid.Height;
+						}
+						else if ((!vid.Width && w) || (!vid.Height && h))
+						{
+							vid.Width = w;
+							vid.Height = h;
+						}
+					}
+#if 0
+					// HDR metadata?
+					else if (OBUtype == ElectraDecodersUtil::AV1Video::EOBUType::OBU_METADATA)
+					{
+						ElectraDecodersUtil::AV1Video::FBitstreamReaderAV1 bs(CurrentOBU, obu_size);
+						const ElectraDecodersUtil::AV1Video::EMetadataType MetadataType = (const ElectraDecodersUtil::AV1Video::EMetadataType)(bs.leb128(nr));
+						switch(MetadataType)
+						{
+							case ElectraDecodersUtil::AV1Video::EMetadataType::METADATA_TYPE_HDR_CLL:
+							{
+								ElectraDecodersUtil::AV1Video::FMetadata_hdr_cll cll;
+								if (cll.ParseFrom(bs))
+								{
+									// create clli box if not already provided
+								}
+								break;
+							}
+							case ElectraDecodersUtil::AV1Video::EMetadataType::METADATA_TYPE_HDR_MDCV:
+							{
+								ElectraDecodersUtil::AV1Video::FMetadata_hdr_mdvc mdcv;
+								if (mdcv.ParseFrom(bs))
+								{
+									// create mdcv box if not already provided
+								}
+								break;
+							}
+						}
+					}
+#endif
+					CurrentOBU += obu_size;
+				}
+			}
+
+			vid.BitDepth = dcr.GetBitDepth();
+			Electra::FCodecTypeFormat::FVideo::FColorInfo ci;
+			const ElectraDecodersUtil::AV1Video::FAV1CodecConfigurationRecord::FColorInfo& av1Col = dcr.GetColorInfo();
+			ci.colourPrimaries = av1Col.colorPrimaries;
+			ci.transferCharacteristics = av1Col.transferCharacteristics;
+			ci.matrixCoefficients = av1Col.matrixCoefficients;
+			ci.videoFullRangeFlag = av1Col.videoFullRangeFlag;
+			Electra::FCodecTypeFormat::FVideo::FColorInfo::FChromaInfo cri;
+			cri.monochrome = av1Col.monochrome;
+			cri.chromaSubsamplingX = av1Col.chromaSubsamplingX;
+			cri.chromaSubsamplingY = av1Col.chromaSubsamplingY;
+			cri.chromaSamplingPosition = av1Col.chromaSamplingPosition;
+			ci.ChromaInfo = cri;
+			vid.OptColorInfo = ci;
+			InOut.RFC6381 = dcr.GetCodecSpecifierRFC6381();
+			InOut.HumanReadableFormatInfo = dcr.GetFormatInfo();
+			return 1;
+		}
+		return -1;
+	}
+	return 0;
+}
+
+static int32 DetermineCodecTypeFromDCR_mp4a(Electra::FCodecTypeFormat& InOut)
+{
+	const TArray<uint8>* esds_Box = InOut.ExtraBoxes.Find(Make4CC('e','s','d','s'));
+	TArray<uint8> NestedESDSBoxData;
+	// With QuickTime files the `esds` box might be stored inside a `wave` box.
+	if (!esds_Box)
+	{
+		if (const TArray<uint8>* wave_Box = InOut.ExtraBoxes.Find(Make4CC('w','a','v','e')))
+		{
+			TArray<TSharedPtr<Electra::UtilitiesMP4::FMP4BoxData, ESPMode::ThreadSafe>> ChildBoxes;
+			if (Electra::UtilitiesMP4::FMP4BoxLocatorReader::LocateAndReadRootBoxesFromBuffer(ChildBoxes, *wave_Box))
+			{
+				if (auto ESDS = ChildBoxes.FindByPredicate([](const TSharedPtr<Electra::UtilitiesMP4::FMP4BoxData, ESPMode::ThreadSafe>& e) { return e->Type == Electra::UtilitiesMP4::MakeBoxAtom('e','s','d','s');} ))
+				{
+					NestedESDSBoxData = (*ESDS)->DataBuffer;
+					esds_Box = &NestedESDSBoxData;
+				}
+			}
+		}
+	}
+
+	if (esds_Box && esds_Box->Num() > 4)
+	{
+		ElectraDecodersUtil::MPEG::FESDescriptor esds;
+		// The ESDescriptorBox is a full box, so we need to skip the first 4 bytes.
+		auto BoxData = MakeConstArrayView<const uint8>(esds_Box->GetData() + 4, esds_Box->Num() - 4);
+		if (!esds.Parse(BoxData))
+		{
+			return -1;
+		}
+		if (esds.GetStreamType() != ElectraDecodersUtil::MPEG::FESDescriptor::FStreamType::AudioStream)
+		{
+			return 0;
+		}
+		if (!InOut.Properties.IsType<Electra::FCodecTypeFormat::FAudio>())
+		{
+			InOut.Properties.Emplace<Electra::FCodecTypeFormat::FAudio>();
+		}
+		Electra::FCodecTypeFormat::FAudio& aud(InOut.Properties.Get<Electra::FCodecTypeFormat::FAudio>());
+
+		uint32 FourCC = 0;
+		int32 AudioObjectType = 0;
+		int32 Layer = 0;
+		int32 Mpeg = 0;
+		switch(esds.GetObjectTypeID())
+		{
+			// ISO/IEC 14496-3
+			case ElectraDecodersUtil::MPEG::FESDescriptor::FObjectTypeID::MPEG4_Audio:
+			{
+				ElectraDecodersUtil::MPEG::FAACDecoderConfigurationRecord dcr;
+				if (!(esds.GetCodecSpecificData().Num() && dcr.Parse(esds.GetCodecSpecificData())))
+				{
+					UE_LOG(LogElectraFormatUtils, Log, TEXT("ESDescriptor lacks mandatory AudioSpecificConfig for ISO/IEC 14496-3 streams"));
+					return -1;
+				}
+				InOut.DCR = BoxData;
+				InOut.CSD = esds.GetCodecSpecificData();
+				InOut.RFC6381 = dcr.GetCodecSpecifierRFC6381();
+
+				uint32 NumChan = (uint32) ElectraDecodersUtil::MPEG::AACUtils::GetNumberOfChannelsFromChannelConfiguration(dcr.ChannelConfiguration);
+				uint32 SmpRate = dcr.ExtSamplingFrequency ? dcr.ExtSamplingFrequency : dcr.SamplingRate;
+				aud.ChannelConfiguration = dcr.ChannelConfiguration;
+				if ((aud.NumChannels == 0) || (aud.NumChannels && NumChan && aud.NumChannels != NumChan))
+				{
+					aud.NumChannels = NumChan;
+				}
+				if ((aud.SampleRate == 0) || (aud.SampleRate && SmpRate && aud.SampleRate != SmpRate))
+				{
+					aud.SampleRate = SmpRate;
+				}
+				AudioObjectType = dcr.ExtAOT ? dcr.ExtAOT : dcr.AOT;
+				switch(AudioObjectType)
+				{
+					case 2:		// AAC-LC
+					case 5:		// SBR
+					case 29:	// PS
+					{
+						FourCC = Make4CC('m','p','4','a');
+						if (InOut.MimeType.IsEmpty())
+						{
+							InOut.MimeType = TEXT("audio/mp4");
+						}
+						break;
+					}
+					case 32:	// Layer-1
+					case 33:	// Layer-2
+					case 34:	// Layer-3
+					{
+						FourCC = Make4CC('m','p','g','a');
+						if (InOut.MimeType.IsEmpty())
+						{
+							InOut.MimeType = TEXT("audio/mpeg");
+						}
+						Layer = AudioObjectType - 31;
+						Mpeg = 1;
+						break;
+					}
+					default:
+					{
+						// not supported
+						return -1;
+					}
+				}
+				break;
+			}
+			// ISO/IEC 13818-7
+			case ElectraDecodersUtil::MPEG::FESDescriptor::FObjectTypeID::MPEG2_AAC_LC:
+			{
+				FourCC = Make4CC('m','p','4','a');
+				AudioObjectType = 2;
+				if (InOut.MimeType.IsEmpty())
+				{
+					InOut.MimeType = TEXT("audio/mp4a");
+				}
+				InOut.RFC6381 = FString::Printf(TEXT("mp4a.%02X"), (uint32)esds.GetObjectTypeID());
+				break;
+			}
+			// ISO/IEC 13818-3
+			case ElectraDecodersUtil::MPEG::FESDescriptor::FObjectTypeID::MPEG2_Audio:
+			// ISO/IEC 11172-3
+			case ElectraDecodersUtil::MPEG::FESDescriptor::FObjectTypeID::MPEG1_Audio:
+			{
+				Mpeg = esds.GetObjectTypeID() == ElectraDecodersUtil::MPEG::FESDescriptor::FObjectTypeID::MPEG1_Audio ? 1 : 2;
+				FourCC = Make4CC('m','p','g','a');
+				if (InOut.MimeType.IsEmpty())
+				{
+					InOut.MimeType = TEXT("audio/mpeg");
+				}
+				InOut.RFC6381 = FString::Printf(TEXT("mp4a.%02X"), (uint32)esds.GetObjectTypeID());
+				break;
+			}
+			default:
+			{
+				return -1;
+			}
+		}
+
+		InOut.Type = Electra::FCodecTypeFormat::EType::Audio;
+		if (InOut.FourCC && InOut.FourCC != FourCC)
+		{
+			UE_LOG(LogElectraFormatUtils, Log, TEXT("FourCC was set as %08x but identifed as %08x from FESDescriptor"), InOut.FourCC, FourCC);
+		}
+		InOut.FourCC = FourCC;
+		MakeHumanReadableRFC_MP4A(InOut, AudioObjectType, Mpeg, Layer);
+		return 1;
+	}
+	return 0;
+}
+
+static int32 DetermineCodecTypeFromDCR_Opus(Electra::FCodecTypeFormat& InOut)
+{
+	const TArray<uint8>* dOps_Box = InOut.ExtraBoxes.Find(Make4CC('d','O','p','s'));
+	if (dOps_Box)
+	{
+		if (!InOut.Properties.IsType<Electra::FCodecTypeFormat::FAudio>())
+		{
+			InOut.Properties.Emplace<Electra::FCodecTypeFormat::FAudio>();
+		}
+		InOut.DCR = *dOps_Box;
+		InOut.RFC6381 = TEXT("opus");
+		InOut.Type = Electra::FCodecTypeFormat::EType::Audio;
+		uint32 FourCC = Make4CC('O','p','u','s');
+		if (InOut.FourCC && InOut.FourCC != FourCC)
+		{
+			UE_LOG(LogElectraFormatUtils, Log, TEXT("FourCC was set as %08x but identifed as %08x"), InOut.FourCC, FourCC);
+		}
+		InOut.FourCC = FourCC;
+		return 1;
+	}
+	return 0;
+}
+
+static int32 DetermineCodecTypeFromDCR_Flac(Electra::FCodecTypeFormat& InOut)
+{
+	const TArray<uint8>* dfLa_Box = InOut.ExtraBoxes.Find(Make4CC('d','f','L','a'));
+	if (dfLa_Box)
+	{
+		if (!InOut.Properties.IsType<Electra::FCodecTypeFormat::FAudio>())
+		{
+			InOut.Properties.Emplace<Electra::FCodecTypeFormat::FAudio>();
+		}
+		// The FLACSpecificBox is a full box, so we need to skip the first 4 bytes.
+		auto BoxData = MakeConstArrayView<const uint8>(dfLa_Box->GetData() + 4, dfLa_Box->Num() - 4);
+		InOut.DCR = BoxData;
+		InOut.RFC6381 = TEXT("flac");
+		InOut.Type = Electra::FCodecTypeFormat::EType::Audio;
+		uint32 FourCC = Make4CC('f','L','a','C');
+		if (InOut.FourCC && InOut.FourCC != FourCC)
+		{
+			UE_LOG(LogElectraFormatUtils, Log, TEXT("FourCC was set as %08x but identifed as %08x"), InOut.FourCC, FourCC);
+		}
+		InOut.FourCC = FourCC;
+		return 1;
+	}
+	return 0;
+}
+
+
+
+static bool DetermineCodecTypeFromDCR(Electra::FCodecTypeFormat& InOut)
+{
+	int32 Result = 0;
+
+	// AVC (H.264)
+	Result = DetermineCodecTypeFromDCR_avcc(InOut);
+	if (Result) return Result < 0 ? false : true;
+	// HEVC (H.265)
+	Result = DetermineCodecTypeFromDCR_hvcc(InOut);
+	if (Result) return Result < 0 ? false : true;
+	// VP8 / VP9
+	Result = DetermineCodecTypeFromDCR_vpcc(InOut);
+	if (Result) return Result < 0 ? false : true;
+	// AV1
+	Result = DetermineCodecTypeFromDCR_av1c(InOut);
+	if (Result) return Result < 0 ? false : true;
+
+	// MP4A
+	Result = DetermineCodecTypeFromDCR_mp4a(InOut);
+	if (Result) return Result < 0 ? false : true;
+	// Opus
+	Result = DetermineCodecTypeFromDCR_Opus(InOut);
+	if (Result) return /*Result < 0 ? false :*/ true;
+	// Flac
+	Result = DetermineCodecTypeFromDCR_Flac(InOut);
+	if (Result) return /*Result < 0 ? false :*/ true;
+
+	return true;
+}
+
+/**
+ * Checks for some additionally provided boxes and extracts metadata of interest.
+ */
+static bool ProcessExtraBoxes(Electra::FCodecTypeFormat& InOut)
+{
+	if (const TArray<uint8>* btrt_Box = InOut.ExtraBoxes.Find(Make4CC('b','t','r','t')))
+	{
+		auto Btrt = StaticCastSharedPtr<Electra::UtilitiesMP4::FMP4BoxBTRT>(Electra::UtilitiesMP4::FMP4BoxBTRT::Create(nullptr, { *btrt_Box }));
+		InOut.Bitrate = Btrt->GetMaxBitrate() ? Btrt->GetMaxBitrate() : InOut.Bitrate;
+		InOut.AverageBitrate = Btrt->GetAverageBitrate() ? Btrt->GetAverageBitrate() : InOut.AverageBitrate;
+	}
+
+	return true;
+}
+
+}
+
+
+
+
+bool FElectraCodecFormatUtilsModularFeature::PrepareCodecTypeFormat(Electra::FCodecTypeFormat& InOutCodecTypeFormat, const FCodecTypeFormatParams& InCodecTypeFormatParams)
+{
+	RemapLegacyFourCC(InOutCodecTypeFormat);
+	SetDefaultMimeTypeAndRFC6381For4CC(InOutCodecTypeFormat);
+	DetermineCodecTypeFromMimeType(InOutCodecTypeFormat);
+	bool bOk = DetermineCodecTypeFromRFC(InOutCodecTypeFormat) && DetermineCodecTypeFromDCR(InOutCodecTypeFormat) && ProcessExtraBoxes(InOutCodecTypeFormat);
+	return bOk;
+}
diff --git a/Engine/Plugins/Media/ElectraCodecs/Source/ElectraDecoders/Private/Utils/ElectraCodecFormatUtils.h b/Engine/Plugins/Media/ElectraCodecs/Source/ElectraDecoders/Private/Utils/ElectraCodecFormatUtils.h
new file mode 100644
--- /dev/null
+++ b/Engine/Plugins/Media/ElectraCodecs/Source/ElectraDecoders/Private/Utils/ElectraCodecFormatUtils.h
@@ -0,0 +1,16 @@
+// Copyright Epic Games, Inc. All Rights Reserved.
+#pragma once
+
+#include "IElectraCodecFormatUtils.h"
+
+class FElectraCodecFormatUtilsModularFeature : public IElectraCodecFormatUtilsModularFeature
+{
+public:
+	static void Startup();
+	static void Shutdown();
+
+	virtual ~FElectraCodecFormatUtilsModularFeature() = default;
+	bool PrepareCodecTypeFormat(Electra::FCodecTypeFormat& InOutCodecTypeFormat, const FCodecTypeFormatParams& InCodecTypeFormatParams) override;
+private:
+	static FElectraCodecFormatUtilsModularFeature* Self;
+};
diff --git a/Engine/Plugins/Media/ElectraCodecs/Source/ElectraDecoders/Private/Utils/ElectraDecodersUtils.cpp b/Engine/Plugins/Media/ElectraCodecs/Source/ElectraDecoders/Private/Utils/ElectraDecodersUtils.cpp
new file mode 100644
--- /dev/null
+++ b/Engine/Plugins/Media/ElectraCodecs/Source/ElectraDecoders/Private/Utils/ElectraDecodersUtils.cpp
@@ -0,0 +1,709 @@
+// Copyright Epic Games, Inc. All Rights Reserved.
+
+#include "ElectraDecodersUtils.h"
+#include "Utilities/UtilitiesMP4.h"
+#include "Utils/Google/ElectraUtilsVPxVideo.h"
+
+
+namespace ElectraDecodersUtil
+{
+
+namespace
+{
+	inline void LexFromStringHex(int32& OutValue, const TCHAR* Buffer)
+	{
+		OutValue = FCString::Strtoi(Buffer, nullptr, 16);
+	}
+
+	inline void LexFromStringHexU64(uint64& OutValue, const TCHAR* Buffer)
+	{
+		OutValue = FCString::Strtoui64(Buffer, nullptr, 16);
+	}
+
+}
+
+
+
+bool ParseMimeTypeWithCodec(FMimeTypeVideoCodecInfo& OutInfo, const FString& InMimeType)
+{
+	return false;
+}
+
+bool ParseMimeTypeWithCodec(FMimeTypeAudioCodecInfo& OutInfo, const FString& InMimeType)
+{
+	return false;
+}
+
+bool ParseCodecMP4A(FMimeTypeAudioCodecInfo& OutInfo, const FString& InCodecFormat)
+{
+	if (!InCodecFormat.StartsWith("mp4a"))
+	{
+		return false;
+	}
+	OutInfo.Codec = TEXT("mp4a");
+	OutInfo.Codec4CC = Make4CC('m','p','4','a');
+	// Object and profile follow?
+	TArray<FString> Parts;
+	InCodecFormat.ParseIntoArray(Parts, TEXT("."));
+	if (Parts.Num() > 1)
+	{
+		// mp4a.xx.d is recognized.
+		OutInfo.ObjectType = FCString::Strtoi(*Parts[1], nullptr, 16);
+		int32 ProfileValue = 0;
+		if (Parts.Num() > 2)
+		{
+			LexFromString(ProfileValue, *Parts[2]);
+		}
+		OutInfo.Profile = ProfileValue;
+	}
+	return true;
+}
+
+bool ParseCodecH264(FMimeTypeVideoCodecInfo& OutInfo, const FString& InCodecFormat)
+{
+	if (InCodecFormat.StartsWith("avc"))
+	{
+		// avc1 and avc3 (inband SPS/PPS) are recognized.
+		if (InCodecFormat.Len() > 3)
+		{
+			// avc1 or avc3 only!
+			if (InCodecFormat[3] != TCHAR('1') && InCodecFormat[3] != TCHAR('3'))
+			{
+				return false;
+			}
+			OutInfo.Codec = TEXT("avc");
+			OutInfo.Codec4CC = InCodecFormat[3] == TCHAR('1') ? Make4CC('a','v','c','1') : Make4CC('a','v','c','3');
+
+			// Profile and level follow?
+			if (InCodecFormat.Len() > 5 && InCodecFormat[4] == TCHAR('.'))
+			{
+				int32 DotPos;
+				InCodecFormat.FindLastChar(TCHAR('.'), DotPos);
+				FString Temp;
+				int32 TempValue;
+				// We recognize the expected format avcC.xxyyzz and for legacy reasons also avcC.xxx.zz
+				if (InCodecFormat.Len() == 11 && DotPos == 4)
+				{
+					Temp = InCodecFormat.Mid(5, 2);
+					LexFromStringHex(TempValue, *Temp);
+					OutInfo.Profile = TempValue;
+					Temp = InCodecFormat.Mid(7, 2);
+					LexFromStringHex(TempValue, *Temp);
+					OutInfo.Constraints = TempValue;
+					Temp = InCodecFormat.Mid(9, 2);
+					LexFromStringHex(TempValue, *Temp);
+					OutInfo.Level = TempValue;
+				}
+				else if (DotPos != INDEX_NONE)
+				{
+					Temp = InCodecFormat.Mid(5, DotPos-5);
+					LexFromString(TempValue, *Temp);
+					OutInfo.Profile = TempValue;
+					Temp = InCodecFormat.Mid(DotPos+1);
+					LexFromString(TempValue, *Temp);
+					OutInfo.Level = TempValue;
+				}
+				else
+				{
+					return false;
+				}
+			}
+			return true;
+		}
+	}
+	return false;
+}
+
+bool ParseCodecH265(FMimeTypeVideoCodecInfo& OutInfo, const FString& InCodecFormat)
+{
+	if (InCodecFormat.StartsWith("hvc1") || InCodecFormat.StartsWith("hev1"))
+	{
+		// hvc1 and hev1 (inband VPS/SPS/PPS) are recognized.
+		if (InCodecFormat.Len() > 4)
+		{
+			OutInfo.Codec = TEXT("hevc");
+			OutInfo.Codec4CC = InCodecFormat.StartsWith("hvc1") ? Make4CC('h','v','c','1') : Make4CC('h','e','v','1');
+
+			FString oti = InCodecFormat;
+			FString Temp;
+			int32 DotPos;
+			if (oti.FindChar(TCHAR('.'), DotPos))
+			{
+				int32 general_profile_space = 0;
+				int32 general_tier_flag = 0;
+				int32 general_profile_idc = 0;
+				int32 general_level_idc = 0;
+				uint32 general_profile_compatibility_flag = 0;
+				uint64 contraint_flags = 0;
+				oti.RightChopInline(DotPos + 1);
+				// optional general_profile_space
+				if (oti[0] == TCHAR('A') || oti[0] == TCHAR('B') || oti[0] == TCHAR('C'))
+				{
+					general_profile_space = oti[0] - TCHAR('A') + 1;
+					oti.RightChopInline(1);
+				}
+				else if (oti[0] == TCHAR('a') || oti[0] == TCHAR('b') || oti[0] == TCHAR('c'))
+				{
+					general_profile_space = oti[0] - TCHAR('a') + 1;
+					oti.RightChopInline(1);
+				}
+				// general_profile_idc
+				if (oti.FindChar(TCHAR('.'), DotPos))
+				{
+					Temp = oti.Left(DotPos);
+					oti.RightChopInline(DotPos + 1);
+					LexFromString(general_profile_idc, *Temp);
+				}
+				// general_profile_compatibility_flags
+				if (oti.FindChar(TCHAR('.'), DotPos))
+				{
+					Temp = oti.Left(DotPos);
+					oti.RightChopInline(DotPos + 1);
+					LexFromString(general_profile_compatibility_flag, *Temp);
+				}
+				// general_tier_flag
+				if (oti[0] != TCHAR('L') && oti[0] != TCHAR('H') && oti[0] != TCHAR('l') && oti[0] != TCHAR('h'))
+				{
+					return false;
+				}
+				else if (oti[0] == TCHAR('H') || oti[0] == TCHAR('h'))
+				{
+					general_tier_flag = 1;
+				}
+				oti.RightChopInline(1);
+				// constraint_flags
+				FString ConstraintFlags;
+				if (oti.FindChar(TCHAR('.'), DotPos))
+				{
+					ConstraintFlags = oti.Mid(DotPos + 1);
+					oti.LeftInline(DotPos);
+					ConstraintFlags.ReplaceInline(TEXT("."), TEXT(""));
+					ConstraintFlags += TEXT("000000000000");
+					ConstraintFlags.LeftInline(12);
+					LexFromStringHexU64(contraint_flags, *ConstraintFlags);
+				}
+				// general_level_idc
+				LexFromString(general_level_idc, *oti);
+
+				OutInfo.Profile = general_profile_idc;
+				OutInfo.Level = general_level_idc;
+				OutInfo.ProfileSpace = general_profile_space;
+				OutInfo.CompatibilityFlags = BitReverse32(general_profile_compatibility_flag);
+				OutInfo.Tier = general_tier_flag;
+				OutInfo.Constraints = contraint_flags;
+				return true;
+			}
+		}
+	}
+	return false;
+}
+
+bool ParseCodecVP8(FMimeTypeVideoCodecInfo& OutInfo, const FString& InCodecFormat, const TArray<uint8>& InvpcCBox)
+{
+	if (InCodecFormat.StartsWith("vp08") || InCodecFormat.StartsWith("vp8"))
+	{
+		FString oti = InCodecFormat;
+		OutInfo.Codec = TEXT("vp8");
+		OutInfo.Codec4CC = Make4CC('v','p','0','8');
+		if (!InvpcCBox.IsEmpty())
+		{
+			// Enough data to represent the `vpcC` box, and is it of version 1?
+			if (InvpcCBox.Num() < 12 || InvpcCBox[0] != 1)
+			{
+				return false;
+			}
+			ElectraDecodersUtil::VPxVideo::FVPCodecConfigurationRecord dcr;
+			if (dcr.Parse(8, MakeConstArrayView<const uint8>(InvpcCBox.GetData() + 4, InvpcCBox.Num() - 4)))
+			{
+				OutInfo.Profile = dcr.GetProfile();
+				OutInfo.Level = dcr.GetLevel();
+				OutInfo.NumBitsLuma = dcr.GetBitDepth();
+				OutInfo.Extras.Emplace<VPxVideo::FVPCodecConfigurationRecord::FColorInfo>(dcr.GetColorInfo());
+				return true;
+			}
+			return false;
+		}
+		if (oti.Len() > 4)
+		{
+			FString Temp;
+			int32 DotPos;
+			if (oti.FindChar(TCHAR('.'), DotPos))
+			{
+				int32 Components[8] {0}, NumComponents=0;
+				oti.RightChopInline(DotPos + 1);
+				while(oti.Len() && NumComponents<UE_ARRAY_COUNT(Components))
+				{
+					if (oti.FindChar(TCHAR('.'), DotPos))
+					{
+						Temp = oti.Left(DotPos);
+						oti.RightChopInline(DotPos + 1);
+						LexFromString(Components[NumComponents++], *Temp);
+					}
+					else
+					{
+						LexFromString(Components[NumComponents++], *oti);
+						break;
+					}
+				}
+				OutInfo.Profile = Components[0];
+				OutInfo.Level = Components[1];
+				OutInfo.NumBitsLuma = Components[2];
+				OutInfo.Extras.Emplace<VPxVideo::FVPCodecConfigurationRecord::FColorInfo>(VPxVideo::FVPCodecConfigurationRecord::FColorInfo());
+				return true;
+			}
+		}
+	}
+	return false;
+}
+
+bool ParseCodecVP9(FMimeTypeVideoCodecInfo& OutInfo, const FString& InCodecFormat, const TArray<uint8>& InvpcCBox)
+{
+	if (InCodecFormat.StartsWith("vp09") || InCodecFormat.StartsWith("vp9"))
+	{
+		auto BuildRFC6381 = [&]() -> FString
+		{
+			VPxVideo::FVPCodecConfigurationRecord::FColorInfo& Col = OutInfo.Extras.Get<VPxVideo::FVPCodecConfigurationRecord::FColorInfo>();
+			FString rfc(FString::Printf(TEXT("vp09.%02d.%02d.%02d.%02d.%02d.%02d.%02d.%02d"), OutInfo.Profile, OutInfo.Level, OutInfo.NumBitsLuma,
+				Col.chromaSubsampling, Col.colourPrimaries, Col.transferCharacteristics, Col.matrixCoefficients, Col.videoFullRangeFlag));
+			if (rfc.EndsWith(TEXT(".01.01.01.01.00")))
+			{
+				rfc.LeftChopInline(15);
+			}
+			return rfc;
+		};
+
+		FString oti = InCodecFormat;
+		OutInfo.Codec = TEXT("vp9");
+		OutInfo.Codec4CC = Make4CC('v','p','0','9');
+		if (!InvpcCBox.IsEmpty())
+		{
+			// Enough data to represent the `vpcC` box, and is it of version 1?
+			if (InvpcCBox.Num() < 12 || InvpcCBox[0] != 1)
+			{
+				return false;
+			}
+			ElectraDecodersUtil::VPxVideo::FVPCodecConfigurationRecord dcr;
+			if (dcr.Parse(9, MakeConstArrayView<const uint8>(InvpcCBox.GetData() + 4, InvpcCBox.Num() - 4)))
+			{
+				OutInfo.Profile = dcr.GetProfile();
+				OutInfo.Level = dcr.GetLevel();
+				OutInfo.NumBitsLuma = dcr.GetBitDepth();
+				OutInfo.Extras.Emplace<VPxVideo::FVPCodecConfigurationRecord::FColorInfo>(dcr.GetColorInfo());
+				OutInfo.RFC6381 = BuildRFC6381();
+				return true;
+			}
+			return false;
+		}
+		if (oti.Len() > 4)
+		{
+			FString Temp;
+			int32 DotPos;
+			if (oti.FindChar(TCHAR('.'), DotPos))
+			{
+				int32 Components[8] {0}, NumComponents=0;
+				oti.RightChopInline(DotPos + 1);
+				while(oti.Len() && NumComponents<UE_ARRAY_COUNT(Components))
+				{
+					if (oti.FindChar(TCHAR('.'), DotPos))
+					{
+						Temp = oti.Left(DotPos);
+						oti.RightChopInline(DotPos + 1);
+						LexFromString(Components[NumComponents++], *Temp);
+					}
+					else
+					{
+						LexFromString(Components[NumComponents++], *oti);
+						break;
+					}
+				}
+				OutInfo.Profile = Components[0];
+				OutInfo.Level = Components[1];
+				OutInfo.NumBitsLuma = Components[2];
+				OutInfo.Extras.Emplace<VPxVideo::FVPCodecConfigurationRecord::FColorInfo>(VPxVideo::FVPCodecConfigurationRecord::FColorInfo());
+				if (NumComponents > 3)
+				{
+					VPxVideo::FVPCodecConfigurationRecord::FColorInfo& ci = OutInfo.Extras.Get<VPxVideo::FVPCodecConfigurationRecord::FColorInfo>();
+					ci.chromaSubsampling = (uint8) Components[3];
+					ci.colourPrimaries = (uint8) Components[4];
+					ci.transferCharacteristics  = (uint8) Components[5];
+					ci.matrixCoefficients = (uint8) Components[6];
+					ci.videoFullRangeFlag = (uint8) Components[7];
+				}
+				OutInfo.RFC6381 = BuildRFC6381();
+				return true;
+			}
+		}
+	}
+	return false;
+}
+
+bool ParseCodecAV1(FMimeTypeVideoCodecInfo& OutInfo, const FString& InCodecFormat, const TArray<uint8>& Inav1CBox, const TArray<uint8>& IncolrBox)
+{
+	if (InCodecFormat.StartsWith("av01"))
+	{
+		auto BuildRFC6381 = [&]() -> FString
+		{
+			AV1Video::FAV1CodecConfigurationRecord::FColorInfo& Col = OutInfo.Extras.Get<AV1Video::FAV1CodecConfigurationRecord::FColorInfo>();
+			FString rfc(FString::Printf(TEXT("av01.%d.%02d%c.%02d.%d.%d%d%d.%02d.%02d.%02d.%d"), OutInfo.Profile, OutInfo.Level, OutInfo.Tier?TCHAR('H'):TCHAR('M'), OutInfo.NumBitsLuma,
+				Col.monochrome, Col.chromaSubsamplingX, Col.chromaSubsamplingY, Col.chromaSamplingPosition,	Col.colorPrimaries, Col.transferCharacteristics, Col.matrixCoefficients, Col.videoFullRangeFlag));
+			if (rfc.EndsWith(TEXT(".0.110.01.01.01.0")))
+			{
+				rfc.LeftChopInline(17);
+			}
+			return rfc;
+		};
+
+		FString oti = InCodecFormat;
+		OutInfo.Codec = TEXT("av1");
+		OutInfo.Codec4CC = Make4CC('a','v','0','1');
+		if (!Inav1CBox.IsEmpty())
+		{
+			// Enough data to represent the `av1C` box, and is it of version 1?
+			if (Inav1CBox.Num() < 4 || Inav1CBox[0] != 0x81)
+			{
+				return false;
+			}
+
+			ElectraDecodersUtil::AV1Video::FAV1CodecConfigurationRecord dcr;
+			if (!dcr.Parse(MakeConstArrayView<const uint8>(Inav1CBox.GetData(), Inav1CBox.Num())))
+			{
+				return false;
+			}
+			OutInfo.Profile = dcr.GetProfile();
+			OutInfo.Level = dcr.GetLevel();
+			OutInfo.Tier = dcr.GetTier();
+			OutInfo.NumBitsLuma = dcr.GetBitDepth();
+			OutInfo.Extras.Emplace<AV1Video::FAV1CodecConfigurationRecord::FColorInfo>(dcr.GetColorInfo());
+			// Is a 'colr' box provided? And does it contain the mandated 'nclx' box?
+			Electra::UtilitiesMP4::FMP4AtomReader boxReader(IncolrBox);
+			uint32 Type;
+			if (boxReader.Read(Type) && Type == Electra::UtilitiesMP4::MakeBoxAtom('n','c','l','x'))
+			{
+				uint16 colour_primaries, transfer_characteristics, matrix_coeffs;
+				uint8 video_full_range_flag = 0;
+				if (boxReader.Read(colour_primaries) && boxReader.Read(transfer_characteristics) && boxReader.Read(matrix_coeffs) && boxReader.Read(video_full_range_flag))
+				{
+					AV1Video::FAV1CodecConfigurationRecord::FColorInfo& ci = OutInfo.Extras.Get<AV1Video::FAV1CodecConfigurationRecord::FColorInfo>();
+					ci.colorPrimaries = colour_primaries;
+					ci.transferCharacteristics = transfer_characteristics;
+					ci.matrixCoefficients = matrix_coeffs;
+					ci.videoFullRangeFlag = video_full_range_flag;
+				}
+			}
+			OutInfo.RFC6381 = BuildRFC6381();
+			return true;
+		}
+		if (oti.Len() > 4)
+		{
+			TArray<FString> Parts;
+			InCodecFormat.ParseIntoArray(Parts, TEXT("."), false);
+			// Profile, level, tier and bitdepth are required.
+			if (Parts.Num() >= 4)
+			{
+				OutInfo.Extras.Emplace<AV1Video::FAV1CodecConfigurationRecord::FColorInfo>(AV1Video::FAV1CodecConfigurationRecord::FColorInfo());
+				AV1Video::FAV1CodecConfigurationRecord::FColorInfo& ci = OutInfo.Extras.Get<AV1Video::FAV1CodecConfigurationRecord::FColorInfo>();
+				for(int32 i=1; i<Parts.Num(); ++i)
+				{
+					switch(i)
+					{
+						case 1:
+						{
+							// Profile. Has to be a single digit.
+							if (Parts[i].Len() > 1)
+							{
+								return false;
+							}
+							LexFromString(OutInfo.Profile, *Parts[i]);
+							break;
+						}
+						case 2:
+						{
+							// Level and tier. Needs to be 2 digits and an 'M' or 'H'.
+							if (Parts[i].Len() != 3 || (Parts[i][2]!=TCHAR('M') && Parts[i][2]!=TCHAR('m') && Parts[i][2]!=TCHAR('H') && Parts[i][2]!=TCHAR('h')))
+							{
+								return false;
+							}
+							OutInfo.Tier = Parts[i][2]==TCHAR('H') || Parts[i][2]==TCHAR('h') ? 1 : 0;
+							Parts[i].LeftChopInline(1);
+							LexFromString(OutInfo.Level, *Parts[i]);
+							break;
+						}
+						case 3:
+						{
+							// Bitdepth, 2 digits
+							if (Parts[i].Len() != 2)
+							{
+								return false;
+							}
+							LexFromString(OutInfo.NumBitsLuma, *Parts[i]);
+							break;
+						}
+						/*
+							From this point on missing or malformed components are not critical. As per the standard:
+								"If any character that is not '.', digits, part of the AV1 4CC, or a tier value is encountered,
+								the string SHALL be interpreted ignoring all the characters starting from that character."
+						*/
+						case 4:
+						{
+							// Monochrome, 1 digit
+							if (Parts[i].Len() != 1)
+							{
+								return true;
+							}
+							LexFromString(ci.monochrome, *Parts[i]);
+							break;
+						}
+						case 5:
+						{
+							// Chroma subsampling, 3 digits
+							if (Parts[i].Len() != 3)
+							{
+								return true;
+							}
+							if (FChar::IsDigit(Parts[i][0])) ci.chromaSubsamplingX = Parts[i][0]-TCHAR('0'); else return true;
+							if (FChar::IsDigit(Parts[i][1])) ci.chromaSubsamplingY = Parts[i][1]-TCHAR('0'); else return true;
+							if (FChar::IsDigit(Parts[i][2])) ci.chromaSamplingPosition = Parts[i][2]-TCHAR('0'); else return true;
+							break;
+						}
+						case 6:
+						{
+							// colorPrimaries.
+							if (Parts[i].Len() != 2)
+							{
+								return false;
+							}
+							LexFromString(ci.colorPrimaries, *Parts[i]);
+							break;
+						}
+						case 7:
+						{
+							// transferCharacteristics.
+							if (Parts[i].Len() != 2)
+							{
+								return false;
+							}
+							LexFromString(ci.transferCharacteristics, *Parts[i]);
+							break;
+						}
+						case 8:
+						{
+							// matrixCoefficients.
+							if (Parts[i].Len() != 2)
+							{
+								return false;
+							}
+							LexFromString(ci.matrixCoefficients, *Parts[i]);
+							break;
+						}
+						case 9:
+						{
+							// videoFullRangeFlag.
+							if (Parts[i].Len() != 1)
+							{
+								return false;
+							}
+							LexFromString(ci.videoFullRangeFlag, *Parts[i]);
+							break;
+						}
+						default:
+						{
+							break;
+						}
+					}
+				}
+				OutInfo.RFC6381 = BuildRFC6381();
+				return true;
+			}
+		}
+	}
+	return false;
+}
+
+
+int64 GetVariantValueSafeI64(const TMap<FString, FVariant>& InFromMap, const FString& InName, int64 InDefaultValue)
+{
+	int64 V = InDefaultValue;
+	const FVariant* Var = InFromMap.Find(InName);
+	if (Var)
+	{
+		if (Var->GetType() == EVariantTypes::Int32)
+		{
+			V = Var->GetValue<int32>();
+		}
+		else if (Var->GetType() == EVariantTypes::Int64)
+		{
+			V = Var->GetValue<int64>();
+		}
+		else if (Var->GetType() == EVariantTypes::UInt64)
+		{
+			V = (int64) Var->GetValue<uint64>();
+		}
+		else if (Var->GetType() == EVariantTypes::UInt32)
+		{
+			V = (int64) Var->GetValue<uint32>();
+		}
+		else if (Var->GetType() == EVariantTypes::Int16)
+		{
+			V = Var->GetValue<int16>();
+		}
+		else if (Var->GetType() == EVariantTypes::Int8)
+		{
+			V = Var->GetValue<int8>();
+		}
+		else if (Var->GetType() == EVariantTypes::UInt16)
+		{
+			V = Var->GetValue<uint16>();
+		}
+		else if (Var->GetType() == EVariantTypes::UInt8)
+		{
+			V = Var->GetValue<uint8>();
+		}
+		else if (Var->GetType() == EVariantTypes::Float)
+		{
+			V = (int64) Var->GetValue<float>();
+		}
+		else if (Var->GetType() == EVariantTypes::Double)
+		{
+			V = (int64) Var->GetValue<double>();
+		}
+		else if (Var->GetType() == EVariantTypes::Bool)
+		{
+			V = Var->GetValue<bool>() ? 1 : 0;
+		}
+	}
+	return V;
+}
+
+
+uint64 GetVariantValueSafeU64(const TMap<FString, FVariant>& InFromMap, const FString& InName, uint64 InDefaultValue)
+{
+	uint64 V = InDefaultValue;
+	const FVariant* Var = InFromMap.Find(InName);
+	if (Var)
+	{
+		if (Var->GetType() == EVariantTypes::Int32)
+		{
+			V = Var->GetValue<int32>();
+		}
+		else if (Var->GetType() == EVariantTypes::Int64)
+		{
+			V = (uint64) Var->GetValue<int64>();
+		}
+		else if (Var->GetType() == EVariantTypes::UInt64)
+		{
+			V = Var->GetValue<uint64>();
+		}
+		else if (Var->GetType() == EVariantTypes::UInt32)
+		{
+			V = Var->GetValue<uint32>();
+		}
+		else if (Var->GetType() == EVariantTypes::Int16)
+		{
+			V = Var->GetValue<int16>();
+		}
+		else if (Var->GetType() == EVariantTypes::Int8)
+		{
+			V = Var->GetValue<int8>();
+		}
+		else if (Var->GetType() == EVariantTypes::UInt16)
+		{
+			V = Var->GetValue<uint16>();
+		}
+		else if (Var->GetType() == EVariantTypes::UInt8)
+		{
+			V = Var->GetValue<uint8>();
+		}
+		else if (Var->GetType() == EVariantTypes::Float)
+		{
+			V = (uint64) Var->GetValue<float>();
+		}
+		else if (Var->GetType() == EVariantTypes::Double)
+		{
+			V = (uint64) Var->GetValue<double>();
+		}
+		else if (Var->GetType() == EVariantTypes::Bool)
+		{
+			V = Var->GetValue<bool>() ? 1 : 0;
+		}
+	}
+	return V;
+}
+
+
+double GetVariantValueSafeDouble(const TMap<FString, FVariant>& InFromMap, const FString& InName, double InDefaultValue)
+{
+	double V = InDefaultValue;
+	const FVariant* Var = InFromMap.Find(InName);
+	if (Var)
+	{
+		if (Var->GetType() == EVariantTypes::Int32)
+		{
+			V = (double) Var->GetValue<int32>();
+		}
+		else if (Var->GetType() == EVariantTypes::Int64)
+		{
+			V = (double) Var->GetValue<int64>();
+		}
+		else if (Var->GetType() == EVariantTypes::UInt32)
+		{
+			V = (double) Var->GetValue<uint32>();
+		}
+		else if (Var->GetType() == EVariantTypes::Int16)
+		{
+			V = (double) Var->GetValue<int16>();
+		}
+		else if (Var->GetType() == EVariantTypes::Int8)
+		{
+			V = (double) Var->GetValue<int8>();
+		}
+		else if (Var->GetType() == EVariantTypes::UInt16)
+		{
+			V = (double) Var->GetValue<uint16>();
+		}
+		else if (Var->GetType() == EVariantTypes::UInt8)
+		{
+			V = (double) Var->GetValue<uint8>();
+		}
+		else if (Var->GetType() == EVariantTypes::Float)
+		{
+			V = (double) Var->GetValue<float>();
+		}
+		else if (Var->GetType() == EVariantTypes::Double)
+		{
+			V = (double) Var->GetValue<double>();
+		}
+		else if (Var->GetType() == EVariantTypes::Bool)
+		{
+			V = Var->GetValue<bool>() ? 1.0 : 0.0;
+		}
+	}
+	return V;
+}
+
+TArray<uint8> GetVariantValueUInt8Array(const TMap<FString, FVariant>& InFromMap, const FString& InName)
+{
+	const FVariant* Var = InFromMap.Find(InName);
+	if (Var)
+	{
+		if (Var->GetType() == EVariantTypes::ByteArray)
+		{
+			return Var->GetValue<TArray<uint8>>();
+		}
+	}
+	return TArray<uint8>();
+}
+
+FString GetVariantValueFString(const TMap<FString, FVariant>& InFromMap, const FString& InName)
+{
+	const FVariant* Var = InFromMap.Find(InName);
+	if (Var)
+	{
+		if (Var->GetType() == EVariantTypes::String)
+		{
+			return Var->GetValue<FString>();
+		}
+	}
+	return FString();
+}
+
+}
diff --git a/Engine/Plugins/Media/ElectraCodecs/Source/ElectraDecoders/Public/IElectraCodecFormatUtils.h b/Engine/Plugins/Media/ElectraCodecs/Source/ElectraDecoders/Public/IElectraCodecFormatUtils.h
new file mode 100644
--- /dev/null
+++ b/Engine/Plugins/Media/ElectraCodecs/Source/ElectraDecoders/Public/IElectraCodecFormatUtils.h
@@ -0,0 +1,40 @@
+// Copyright Epic Games, Inc. All Rights Reserved.
+
+#pragma once
+
+#include "CoreTypes.h"
+#include "Features/IModularFeature.h"
+#include "Misc/Optional.h"
+
+#include "CodecTypeFormat.h"
+
+
+class IElectraCodecFormatUtilsModularFeature : public IModularFeature
+{
+public:
+	virtual ~IElectraCodecFormatUtilsModularFeature() = default;
+
+	static FName GetModularFeatureName()
+	{
+		static FName FeatureName = FName(TEXT("ElectraCodecFormatUtils"));
+		return FeatureName;
+	}
+
+	struct FCodecTypeFormatParams
+	{
+		// No members right now.
+	};
+
+	/**
+	 * Tries to prepare the given codec type format format with as much information as possible.
+	 * On input:
+	 *   - set at least the .FourCC member OR the mime type with a codecs parameter
+	 *   - set the .RFC6381 member if possible and not already provided with the mime type
+	 *   - provide any additional boxes in .ExtraBoxes in ISO/IEC 14496-12 binary box format
+	 *      (these should be all boxes from the `stsd` or similar container box)
+	 * On output:
+	 *   - true, with as many filled in values as could be derived from the given input.
+	 *   - false if no information could be derived from the given input.
+	 */
+	virtual bool PrepareCodecTypeFormat(Electra::FCodecTypeFormat& InOutCodecTypeFormat, const FCodecTypeFormatParams& InCodecTypeFormatParams) = 0;
+};
diff --git a/Engine/Plugins/Media/ElectraCodecs/Source/ElectraDecoders/Public/Utils/MPEG/ElectraUtilsMPEGAudio.h b/Engine/Plugins/Media/ElectraCodecs/Source/ElectraDecoders/Public/Utils/MPEG/ElectraUtilsMPEGAudio.h
--- a/Engine/Plugins/Media/ElectraCodecs/Source/ElectraDecoders/Public/Utils/MPEG/ElectraUtilsMPEGAudio.h
+++ b/Engine/Plugins/Media/ElectraCodecs/Source/ElectraDecoders/Public/Utils/MPEG/ElectraUtilsMPEGAudio.h
@@ -38,8 +38,10 @@ namespace ElectraDecodersUtil
 				H264 = 0x21,
 				H264_ParameterSets = 0x22,
 				H265 = 0x23,
-				MPEG4_Audio = 0x40,
-				MPEG1_Audio = 0x6b
+				MPEG4_Audio = 0x40,		// ISO/IEC 14496-3
+				MPEG2_AAC_LC = 0x67,	// ISO/IEC 13818-7
+				MPEG2_Audio = 0x69,		// ISO/IEC 13818-3
+				MPEG1_Audio = 0x6b		// ISO/IEC 11172-3
 			};
 			enum class FStreamType
 			{
diff --git a/Engine/Plugins/Media/ElectraCodecs/Source/ElectraDecoders/Public/Utils/MPEG/ElectraUtilsMPEGVideo_H264.h b/Engine/Plugins/Media/ElectraCodecs/Source/ElectraDecoders/Public/Utils/MPEG/ElectraUtilsMPEGVideo_H264.h
--- a/Engine/Plugins/Media/ElectraCodecs/Source/ElectraDecoders/Public/Utils/MPEG/ElectraUtilsMPEGVideo_H264.h
+++ b/Engine/Plugins/Media/ElectraCodecs/Source/ElectraDecoders/Public/Utils/MPEG/ElectraUtilsMPEGVideo_H264.h
@@ -633,6 +633,13 @@ namespace ElectraDecodersUtil
 				inline uint8 GetAVCLevelIndication() const
 				{ return AVCLevelIndication; }
 
+				inline uint8 GetChromaFormat() const
+				{ return ChromaFormat; }
+				inline uint8 GetBitDepthLuma() const
+				{ return 8 + BitDepthLumaMinus8; }
+				inline uint8 GetBitDepthChroma() const
+				{ return 8+ BitDepthChromaMinus8; }
+
 				inline FString GetCodecSpecifierRFC6381(const TCHAR* InPrefix) const
 				{ return FString::Printf(TEXT("%s.%02x%02x%02x"), InPrefix, GetAVCProfileIndication(), GetProfileCompatibility(), GetAVCLevelIndication()); }
 
diff --git a/Engine/Plugins/Media/ElectraPlayer/Source/ElectraProtron/Private/Player/ElectraProtronPlayerControls.cpp b/Engine/Plugins/Media/ElectraPlayer/Source/ElectraProtron/Private/Player/ElectraProtronPlayerControls.cpp
--- a/Engine/Plugins/Media/ElectraPlayer/Source/ElectraProtron/Private/Player/ElectraProtronPlayerControls.cpp
+++ b/Engine/Plugins/Media/ElectraPlayer/Source/ElectraProtron/Private/Player/ElectraProtronPlayerControls.cpp
@@ -5,7 +5,7 @@
 #include "IElectraCodecFactoryModule.h"
 #include "IElectraCodecFactory.h"
 #include "IElectraDecoderFeaturesAndOptions.h"
-#include "TrackFormatInfo.h"
+#include "CodecTypeFormat.h"
 
 
 TRangeSet<float> FElectraProtronPlayer::FImpl::GetSupportedRates(EMediaRateThinning InThinning)
@@ -23,10 +23,10 @@ TRangeSet<float> FElectraProtronPlayer::FImpl::GetSupportedRates(EMediaRateThinn
 		TempRates.Add(TRange<float>(1.0f));
 
 		// No audio?
-		if (TrackSelection.SelectedTrackIndex[CodecTypeIndex(ElectraProtronUtils::FCodecInfo::EType::Audio)] == -1)
+		if (TrackSelection.SelectedTrackIndex[CodecTypeIndex(Electra::FCodecTypeFormat::EType::Audio)] == -1)
 		{
 			// Video selected?
-			const auto TypeIdx = CodecTypeIndex(ElectraProtronUtils::FCodecInfo::EType::Video);
+			const auto TypeIdx = CodecTypeIndex(Electra::FCodecTypeFormat::EType::Video);
 			if (TrackSelection.SelectedTrackIndex[TypeIdx] != -1)
 			{
 				auto Track = Tracks[UsableTrackArrayIndicesByType[TypeIdx][TrackSelection.SelectedTrackIndex[TypeIdx]]];
@@ -76,8 +76,8 @@ void FElectraProtronPlayer::FImpl::HandleActiveTrackChanges()
 			}
 		};
 
-		auto videoCI = CodecTypeIndex(ElectraProtronUtils::FCodecInfo::EType::Video);
-		auto audioCI = CodecTypeIndex(ElectraProtronUtils::FCodecInfo::EType::Audio);
+		auto videoCI = CodecTypeIndex(Electra::FCodecTypeFormat::EType::Video);
+		auto audioCI = CodecTypeIndex(Electra::FCodecTypeFormat::EType::Audio);
 		// See GetSupportedRates() above. When we are currently playing in reverse and switching tracks would result in
 		// reverse playback being disabled, we need to either set the rate to forward play or pause.
 		if (CurrentRate < 0.0f || SharedPlayParams->PlaybackDirection < 0.0f)
@@ -146,8 +146,8 @@ void FElectraProtronPlayer::FImpl::HandleSeekRequest(const FSeekRequest& InSeekR
 
 
 	CurrentPlayPosTime = InSeekRequest.NewTime;
-	UpdateTrackLoader(CodecTypeIndex(ElectraProtronUtils::FCodecInfo::EType::Video));
-	UpdateTrackLoader(CodecTypeIndex(ElectraProtronUtils::FCodecInfo::EType::Audio));
+	UpdateTrackLoader(CodecTypeIndex(Electra::FCodecTypeFormat::EType::Video));
+	UpdateTrackLoader(CodecTypeIndex(Electra::FCodecTypeFormat::EType::Audio));
 
 	VideoDecoderThread.SetTime(InSeekRequest.NewTime, InSeekRequest.NewSequenceIndex, InSeekRequest.NewLoopIndex);
 	VideoDecoderThread.ResumeAfterSeek();
diff --git a/Engine/Plugins/Media/ElectraPlayer/Source/ElectraProtron/Private/Player/ElectraProtronPlayerDecoder.cpp b/Engine/Plugins/Media/ElectraPlayer/Source/ElectraProtron/Private/Player/ElectraProtronPlayerDecoder.cpp
--- a/Engine/Plugins/Media/ElectraPlayer/Source/ElectraProtron/Private/Player/ElectraProtronPlayerDecoder.cpp
+++ b/Engine/Plugins/Media/ElectraPlayer/Source/ElectraProtron/Private/Player/ElectraProtronPlayerDecoder.cpp
@@ -9,7 +9,7 @@
 #include "Utils/MPEG/ElectraUtilsMPEGVideo_H265.h"
 #include "Utils/MPEG/ElectraUtilsMPEGVideo.h"
 #include "Utils/MPEG/ElectraUtilsMPEGAudio.h"
-#include "TrackFormatInfo.h"
+#include "CodecTypeFormat.h"
 
 #include "IElectraDecoder.h"
 #include "IElectraDecoderOutputVideo.h"
@@ -28,12 +28,12 @@ void FElectraProtronPlayer::FImpl::FDecoderThread::StartThread(const FOpenParam&
 	Params = InParam;
 	SharedPlayParams = InSharedPlayParams;
 	// Video decoder does not need to hold the reference to the audio sample pool.
-	if (DecoderTypeIndex == CodecTypeIndex(ElectraProtronUtils::FCodecInfo::EType::Video))
+	if (DecoderTypeIndex == CodecTypeIndex(Electra::FCodecTypeFormat::EType::Video))
 	{
 		Params.AudioSamplePool.Reset();
 	}
 	// Audio decoder does not need to hold the reference to the texture pool.
-	else if (DecoderTypeIndex == CodecTypeIndex(ElectraProtronUtils::FCodecInfo::EType::Audio))
+	else if (DecoderTypeIndex == CodecTypeIndex(Electra::FCodecTypeFormat::EType::Audio))
 	{
 		Params.TexturePool.Reset();
 	}
@@ -327,7 +327,7 @@ void FElectraProtronPlayer::FImpl::FDecoderThread::Flush(const TSharedPtr<FMedia
 void FElectraProtronPlayer::FImpl::FDecoderThread::UpdateTrackSampleDurationMap()
 {
 	// Only do this for video tracks.
-	if (DecoderTypeIndex != CodecTypeIndex(ElectraProtronUtils::FCodecInfo::EType::Video))
+	if (DecoderTypeIndex != CodecTypeIndex(Electra::FCodecTypeFormat::EType::Video))
 	{
 		return;
 	}
@@ -507,7 +507,7 @@ void FElectraProtronPlayer::FImpl::FDecoderThread::FlushForEndOrLooping()
 	// Non keyframe-only formats, and all audio formats, require that we drain the decoder to get
 	// all the pending output before we start decoding the next (loop) sequence.
 	bool bNeedsDrain = true;
-	if (DecoderTypeIndex == CodecTypeIndex(ElectraProtronUtils::FCodecInfo::EType::Video) &&
+	if (DecoderTypeIndex == CodecTypeIndex(Electra::FCodecTypeFormat::EType::Video) &&
 		TrackSampleBuffer->TrackAndCodecInfo->bIsKeyframeOnlyFormat)
 	{
 		bNeedsDrain = false;
@@ -771,11 +771,11 @@ void FElectraProtronPlayer::FImpl::FDecoderThread::HandleOutputFrame()
 				// No output although advertised?
 				return;
 			}
-			if (DecoderTypeIndex == CodecTypeIndex(ElectraProtronUtils::FCodecInfo::EType::Video) && Out->GetType() == IElectraDecoderOutput::EType::Video)
+			if (DecoderTypeIndex == CodecTypeIndex(Electra::FCodecTypeFormat::EType::Video) && Out->GetType() == IElectraDecoderOutput::EType::Video)
 			{
 				CurrentDecoderOutput = MoveTemp(Out);
 			}
-			else if (DecoderTypeIndex == CodecTypeIndex(ElectraProtronUtils::FCodecInfo::EType::Audio) && Out->GetType() == IElectraDecoderOutput::EType::Audio)
+			else if (DecoderTypeIndex == CodecTypeIndex(Electra::FCodecTypeFormat::EType::Audio) && Out->GetType() == IElectraDecoderOutput::EType::Audio)
 			{
 				CurrentDecoderOutput = MoveTemp(Out);
 			}
@@ -821,7 +821,7 @@ void FElectraProtronPlayer::FImpl::FDecoderThread::HandleOutputFrame()
 				{
 					AdjustToTime = X;
 					SeekTimeToHandleTo.Reset();
-					if (DecoderTypeIndex == CodecTypeIndex(ElectraProtronUtils::FCodecInfo::EType::Video) && (SeekTimeNumFramesDecoded + SeekTimeNumFramesSkipped) > 1)
+					if (DecoderTypeIndex == CodecTypeIndex(Electra::FCodecTypeFormat::EType::Video) && (SeekTimeNumFramesDecoded + SeekTimeNumFramesSkipped) > 1)
 					{
 						UE_LOG(LogElectraProtron, Log, TEXT("Processed %d frames to reach seek point, of which %d could be skipped"), SeekTimeNumFramesDecoded + SeekTimeNumFramesSkipped, SeekTimeNumFramesSkipped);
 					}
@@ -838,7 +838,7 @@ void FElectraProtronPlayer::FImpl::FDecoderThread::HandleOutputFrame()
 		FTimespan FirstSampleEffectiveStartTime = PlaybackRange.GetLowerBoundValue();
 		FTimespan LatestSampleEffectiveEndTime = PlaybackRange.GetUpperBoundValue();
 
-		if (DecoderTypeIndex == CodecTypeIndex(ElectraProtronUtils::FCodecInfo::EType::Video))
+		if (DecoderTypeIndex == CodecTypeIndex(Electra::FCodecTypeFormat::EType::Video))
 		{
 			AdjustForSeek();
 
@@ -973,7 +973,7 @@ void FElectraProtronPlayer::FImpl::FDecoderThread::HandleOutputFrame()
 				InputForCurrentDecoderOutput.Reset();
 			}
 		}
-		else if (DecoderTypeIndex == CodecTypeIndex(ElectraProtronUtils::FCodecInfo::EType::Audio))
+		else if (DecoderTypeIndex == CodecTypeIndex(Electra::FCodecTypeFormat::EType::Audio))
 		{
 			AdjustForSeek();
 
@@ -1132,11 +1132,11 @@ bool FElectraProtronPlayer::FImpl::FDecoderThread::CreateDecoder()
 	CurrentColorimetry.Reset();
 	CurrentHDR.Reset();
 	TMap<FString, FVariant> DecoderCfgParams;
-	const ElectraProtronUtils::FCodecInfo& ci(TrackSampleBuffer->TrackAndCodecInfo->CodecInfo);
-	if (DecoderTypeIndex == CodecTypeIndex(ElectraProtronUtils::FCodecInfo::EType::Video))
+	const Electra::FCodecTypeFormat& ci(TrackSampleBuffer->TrackAndCodecInfo->CodecInfo);
+	if (DecoderTypeIndex == CodecTypeIndex(Electra::FCodecTypeFormat::EType::Video))
 	{
-		check(ci.Type == ElectraProtronUtils::FCodecInfo::EType::Video);
-		const ElectraProtronUtils::FCodecInfo::FVideo& vi(ci.Properties.Get<ElectraProtronUtils::FCodecInfo::FVideo>());
+		check(ci.Type == Electra::FCodecTypeFormat::EType::Video);
+		const Electra::FCodecTypeFormat::FVideo& vi(ci.Properties.Get<Electra::FCodecTypeFormat::FVideo>());
 		DecoderCfgParams.Add(TEXT("width"), FVariant((uint32) vi.Width));
 		DecoderCfgParams.Add(TEXT("height"), FVariant((uint32) vi.Height));
 		DecoderCfgParams.Add(TEXT("max_width"), FVariant((uint32) vi.Width));
@@ -1152,9 +1152,9 @@ bool FElectraProtronPlayer::FImpl::FDecoderThread::CreateDecoder()
 		DecoderCfgParams.Add(TEXT("max_output_buffers"), FVariant(kNumBufs));
 		if (vi.FrameRate.IsValid())
 		{
-			DecoderCfgParams.Add(TEXT("fps"), FVariant((double)vi.FrameRate.GetAsDouble()));
-			DecoderCfgParams.Add(TEXT("fps_n"), FVariant((int64)vi.FrameRate.GetNumerator()));
-			DecoderCfgParams.Add(TEXT("fps_d"), FVariant((uint32)vi.FrameRate.GetDenominator()));
+			DecoderCfgParams.Add(TEXT("fps"), FVariant((double)vi.FrameRate.AsDecimal()));
+			DecoderCfgParams.Add(TEXT("fps_n"), FVariant((int64)vi.FrameRate.Numerator));
+			DecoderCfgParams.Add(TEXT("fps_d"), FVariant((uint32)vi.FrameRate.Denominator));
 		}
 		else
 		{
@@ -1163,10 +1163,10 @@ bool FElectraProtronPlayer::FImpl::FDecoderThread::CreateDecoder()
 			DecoderCfgParams.Add(TEXT("fps_d"), FVariant((uint32)1));
 		}
 	}
-	else if (DecoderTypeIndex == CodecTypeIndex(ElectraProtronUtils::FCodecInfo::EType::Audio))
+	else if (DecoderTypeIndex == CodecTypeIndex(Electra::FCodecTypeFormat::EType::Audio))
 	{
-		check(ci.Type == ElectraProtronUtils::FCodecInfo::EType::Audio);
-		const ElectraProtronUtils::FCodecInfo::FAudio& ai(ci.Properties.Get<ElectraProtronUtils::FCodecInfo::FAudio>());
+		check(ci.Type == Electra::FCodecTypeFormat::EType::Audio);
+		const Electra::FCodecTypeFormat::FAudio& ai(ci.Properties.Get<Electra::FCodecTypeFormat::FAudio>());
 		DecoderCfgParams.Add(TEXT("channel_configuration"), FVariant((uint32) ai.ChannelConfiguration));
 		DecoderCfgParams.Add(TEXT("num_channels"), FVariant((int32) ai.NumChannels));
 		DecoderCfgParams.Add(TEXT("sample_rate"), FVariant((int32) ai.SampleRate));
@@ -1206,7 +1206,7 @@ bool FElectraProtronPlayer::FImpl::FDecoderThread::CreateDecoder()
 
 	// Create platform specific resources to be used with the new decoder.
 	check(DecoderPlatformResource == nullptr);
-	if (DecoderTypeIndex == CodecTypeIndex(ElectraProtronUtils::FCodecInfo::EType::Video))
+	if (DecoderTypeIndex == CodecTypeIndex(Electra::FCodecTypeFormat::EType::Video))
 	{
 		DecoderPlatformResource = FElectraDecodersPlatformResources::CreatePlatformVideoResource(this, DecoderCfgParams);
 	}
diff --git a/Engine/Plugins/Media/ElectraPlayer/Source/ElectraProtron/Private/Player/ElectraProtronPlayerImpl.cpp b/Engine/Plugins/Media/ElectraPlayer/Source/ElectraProtron/Private/Player/ElectraProtronPlayerImpl.cpp
--- a/Engine/Plugins/Media/ElectraPlayer/Source/ElectraProtron/Private/Player/ElectraProtronPlayerImpl.cpp
+++ b/Engine/Plugins/Media/ElectraPlayer/Source/ElectraProtron/Private/Player/ElectraProtronPlayerImpl.cpp
@@ -4,14 +4,16 @@
 #include "ElectraProtronPrivate.h"
 #include "IElectraCodecFactoryModule.h"
 #include "IElectraCodecFactory.h"
+#include "IElectraCodecFormatUtils.h"
 #include "IElectraDecoderFeaturesAndOptions.h"
 #include "Utils/MPEG/ElectraUtilsMPEGVideo_H264.h"
 #include "Utils/MPEG/ElectraUtilsMPEGVideo_H265.h"
 #include "Utils/Google/ElectraUtilsVPxVideo.h"
 #include "Utils/AOMedia/ElectraUtilsAV1Video.h"
 #include "Utils/MPEG/ElectraUtilsMPEGAudio.h"
-#include "TrackFormatInfo.h"
+#include "CodecTypeFormat.h"
 #include "Core/MediaThreads.h"
+#include "Features/IModularFeatures.h"
 
 
 namespace ElectraProtronOptionNames
@@ -37,6 +39,13 @@ FElectraProtronPlayer::FImpl::FImpl()
 	const int32 kVideoFramesToCacheAhead = 8;
 	const int32 kVideoFramesToCacheBehind = 8;
 	CurrentSampleQueueInterface = MakeShared<FSampleQueueInterface, ESPMode::ThreadSafe>(kVideoFramesToCacheAhead, kVideoFramesToCacheBehind);
+
+	if (!CodecFormatUtils)
+	{
+		IModularFeatures::FScopedLockModularFeatureList ScopedLockModularFeatureList;
+		TArray<IElectraCodecFormatUtilsModularFeature*> FormatUtils = IModularFeatures::Get().GetModularFeatureImplementations<IElectraCodecFormatUtilsModularFeature>(IElectraCodecFormatUtilsModularFeature::GetModularFeatureName());
+		CodecFormatUtils = !FormatUtils.IsEmpty() ? FormatUtils[0] : nullptr;
+	}
 }
 
 FElectraProtronPlayer::FImpl::~FImpl()
@@ -90,8 +99,8 @@ FString FElectraProtronPlayer::FImpl::GetLastError()
 
 bool FElectraProtronPlayer::FImpl::HasReachedEnd()
 {
-	const bool bIsVideoActive = TrackSelection.ActiveTrackIndex[CodecTypeIndex(ElectraProtronUtils::FCodecInfo::EType::Video)] != -1;
-	const bool bIsAudioActive = TrackSelection.ActiveTrackIndex[CodecTypeIndex(ElectraProtronUtils::FCodecInfo::EType::Audio)] != -1;
+	const bool bIsVideoActive = TrackSelection.ActiveTrackIndex[CodecTypeIndex(Electra::FCodecTypeFormat::EType::Video)] != -1;
+	const bool bIsAudioActive = TrackSelection.ActiveTrackIndex[CodecTypeIndex(Electra::FCodecTypeFormat::EType::Audio)] != -1;
 
 	bool bAllReachedEnd = true;
 	if (bIsVideoActive && !VideoDecoderThread.HasReachedEnd())
@@ -213,6 +222,12 @@ uint32 FElectraProtronPlayer::FImpl::Run()
  */
 void FElectraProtronPlayer::FImpl::InternalOpen(const FString& InFilename)
 {
+	if (!CodecFormatUtils)
+	{
+		LastErrorMessage = TEXT("No registered IElectraCodecFormatUtilsModularFeature found. Will not be able to determine codecs.");
+		return;
+	}
+
 	// Open the file.
 	TSharedPtr<Electra::IFileDataReader, ESPMode::ThreadSafe> Reader = Electra::IFileDataReader::Create();
 	if (!Reader->Open(InFilename))
@@ -389,32 +404,32 @@ void FElectraProtronPlayer::FImpl::InternalOpen(const FString& InFilename)
 		// Take note of the track with the longest duration (any track, even unsupported ones)
 		LongestTrackDuration = LongestTrackDuration < TrackDuration ? TrackDuration : LongestTrackDuration;
 
-		ElectraProtronUtils::FCodecInfo& CodecInfo(Tracks[i]->CodecInfo);
+		Electra::FCodecTypeFormat& CodecInfo(Tracks[i]->CodecInfo);
 		GetTrackCodecInfo(CodecInfo, ThisTrakBox, TrackID);
 
 		// Timecode tracks are inherently usable.
-		if (CodecInfo.Type == ElectraProtronUtils::FCodecInfo::EType::Timecode)
+		if (CodecInfo.Type == Electra::FCodecTypeFormat::EType::Timecode)
 		{
 			Tracks[i]->bIsUsable = true;
 			Tracks[i]->bIsKeyframeOnlyFormat = true;
 			++NumTimecodeTracks;
 		}
 		// Check with the decoder factory if this format can be decoded.
-		else if (CodecInfo.Type == ElectraProtronUtils::FCodecInfo::EType::Video || CodecInfo.Type == ElectraProtronUtils::FCodecInfo::EType::Audio)
+		else if (CodecInfo.Type == Electra::FCodecTypeFormat::EType::Video || CodecInfo.Type == Electra::FCodecTypeFormat::EType::Audio)
 		{
 			TMap<FString, FVariant> Params;
 			switch(CodecInfo.Type)
 			{
-				case ElectraProtronUtils::FCodecInfo::EType::Video:
+				case Electra::FCodecTypeFormat::EType::Video:
 				{
-					const ElectraProtronUtils::FCodecInfo::FVideo& Video(CodecInfo.Properties.Get<ElectraProtronUtils::FCodecInfo::FVideo>());
+					const Electra::FCodecTypeFormat::FVideo& Video(CodecInfo.Properties.Get<Electra::FCodecTypeFormat::FVideo>());
 					Params.Add(TEXT("width"), FVariant(Video.Width));
 					Params.Add(TEXT("height"), FVariant(Video.Height));
 					if (Video.FrameRate.IsValid())
 					{
-						Params.Add(TEXT("fps"), FVariant((double)Video.FrameRate.GetAsDouble()));
-						Params.Add(TEXT("fps_n"), FVariant((int64)Video.FrameRate.GetNumerator()));
-						Params.Add(TEXT("fps_d"), FVariant((uint32)Video.FrameRate.GetDenominator()));
+						Params.Add(TEXT("fps"), FVariant((double)Video.FrameRate.AsDecimal()));
+						Params.Add(TEXT("fps_n"), FVariant((int64)Video.FrameRate.Numerator));
+						Params.Add(TEXT("fps_d"), FVariant((uint32)Video.FrameRate.Denominator));
 					}
 					else
 					{
@@ -424,9 +439,9 @@ void FElectraProtronPlayer::FImpl::InternalOpen(const FString& InFilename)
 					}
 					break;
 				}
-				case ElectraProtronUtils::FCodecInfo::EType::Audio:
+				case Electra::FCodecTypeFormat::EType::Audio:
 				{
-					const ElectraProtronUtils::FCodecInfo::FAudio& Audio(CodecInfo.Properties.Get<ElectraProtronUtils::FCodecInfo::FAudio>());
+					const Electra::FCodecTypeFormat::FAudio& Audio(CodecInfo.Properties.Get<Electra::FCodecTypeFormat::FAudio>());
 					Params.Add(TEXT("channel_configuration"), FVariant(Audio.ChannelConfiguration));
 					Params.Add(TEXT("num_channels"), FVariant((int32) Audio.NumChannels));
 					Params.Add(TEXT("sample_rate"), FVariant((int32) Audio.SampleRate));
@@ -451,7 +466,7 @@ void FElectraProtronPlayer::FImpl::InternalOpen(const FString& InFilename)
 				TMap<FString, FVariant> ConfigOptions;
 				Factory->GetConfigurationOptions(ConfigOptions);
 				// Every non-video format is assumed to be keyframe-only. For video we ask the factory.
-				Tracks[i]->bIsKeyframeOnlyFormat = CodecInfo.Type != ElectraProtronUtils::FCodecInfo::EType::Video || !!ElectraDecodersUtil::GetVariantValueSafeI64(FormatInfo, IElectraDecoderFormatInfo::IsEveryFrameKeyframe, 0);
+				Tracks[i]->bIsKeyframeOnlyFormat = CodecInfo.Type != Electra::FCodecTypeFormat::EType::Video || !!ElectraDecodersUtil::GetVariantValueSafeI64(FormatInfo, IElectraDecoderFormatInfo::IsEveryFrameKeyframe, 0);
 				// See if the decoder provides a nice name for the format used.
 				Tracks[i]->HumanReadableCodecFormat = ElectraDecodersUtil::GetVariantValueFString(FormatInfo, IElectraDecoderFormatInfo::HumanReadableFormatName);
 				if (Tracks[i]->HumanReadableCodecFormat.IsEmpty())
@@ -509,7 +524,7 @@ void FElectraProtronPlayer::FImpl::InternalOpen(const FString& InFilename)
 			// We need to check every track, even the ones we cannot use. Otherwise we
 			// may think the timecode is not referenced and use it for everything!
 			// We do not check if a timecode track references another one as this would be silly.
-			if (Tracks[TrkNum]->CodecInfo.Type != ElectraProtronUtils::FCodecInfo::EType::Timecode)
+			if (Tracks[TrkNum]->CodecInfo.Type != Electra::FCodecTypeFormat::EType::Timecode)
 			{
 				auto Tref = Tracks[TrkNum]->TrackBox->FindBoxRecursive<Electra::UtilitiesMP4::FMP4BoxTREF>(Electra::UtilitiesMP4::MakeBoxAtom('t','r','e','f'), 1);
 				if (Tref.IsValid())
@@ -562,14 +577,14 @@ void FElectraProtronPlayer::FImpl::InternalOpen(const FString& InFilename)
 			if (NumTimecodeTracks == 1)
 			{
 				// Which track is the timecode?
-				TSharedPtr<FTrackInfo, ESPMode::ThreadSafe>* TimecodeTrack = Tracks.FindByPredicate([](const TSharedPtr<FTrackInfo, ESPMode::ThreadSafe>& In){ return In->CodecInfo.Type == ElectraProtronUtils::FCodecInfo::EType::Timecode; });
+				TSharedPtr<FTrackInfo, ESPMode::ThreadSafe>* TimecodeTrack = Tracks.FindByPredicate([](const TSharedPtr<FTrackInfo, ESPMode::ThreadSafe>& In){ return In->CodecInfo.Type == Electra::FCodecTypeFormat::EType::Timecode; });
 				check(TimecodeTrack); // there has to be one, otherwise we would not even get here, but for safeties sake...
 				if (TimecodeTrack)
 				{
 					// We set this up for video tracks only.
 					for(int32 i=0; i<Tracks.Num(); ++i)
 					{
-						if (Tracks[i]->CodecInfo.Type != ElectraProtronUtils::FCodecInfo::EType::Video)
+						if (Tracks[i]->CodecInfo.Type != Electra::FCodecTypeFormat::EType::Video)
 						{
 							continue;
 						}
@@ -643,8 +658,8 @@ void FElectraProtronPlayer::FImpl::InternalOpen(const FString& InFilename)
 	}
 
 	// Neither video nor audio?
-	if (UsableTrackArrayIndicesByType[CodecTypeIndex(ElectraProtronUtils::FCodecInfo::EType::Video)].Num() == 0 &&
-		UsableTrackArrayIndicesByType[CodecTypeIndex(ElectraProtronUtils::FCodecInfo::EType::Audio)].Num() == 0)
+	if (UsableTrackArrayIndicesByType[CodecTypeIndex(Electra::FCodecTypeFormat::EType::Video)].Num() == 0 &&
+		UsableTrackArrayIndicesByType[CodecTypeIndex(Electra::FCodecTypeFormat::EType::Audio)].Num() == 0)
 	{
 		LastErrorMessage = FString::Printf(TEXT("This file contains no playable video or audio tracks."));
 		return;
@@ -666,7 +681,7 @@ void FElectraProtronPlayer::FImpl::InternalOpen(const FString& InFilename)
 		}
 
 		// If this is a timecode track we may want to read in the first timecode.
-		if (Config.bReadFirstTimecode && Tracks[i]->CodecInfo.Type == ElectraProtronUtils::FCodecInfo::EType::Timecode &&
+		if (Config.bReadFirstTimecode && Tracks[i]->CodecInfo.Type == Electra::FCodecTypeFormat::EType::Timecode &&
 			Tracks[i]->CodecInfo.FourCC == Electra::UtilitiesMP4::MakeBoxAtom('t','m','c','d'))
 		{
 			auto It = Tracks[i]->MP4Track->CreateIterator();
@@ -694,7 +709,7 @@ void FElectraProtronPlayer::FImpl::InternalOpen(const FString& InFilename)
 				return;
 			}
 			// Get the timecode description from the codec info.
-			ElectraProtronUtils::FCodecInfo::FTMCDTimecode TimecodeInfo(Tracks[i]->CodecInfo.Properties.Get<ElectraProtronUtils::FCodecInfo::FTMCDTimecode>());
+			Electra::FCodecTypeFormat::FTMCDTimecode TimecodeInfo(Tracks[i]->CodecInfo.Properties.Get<Electra::FCodecTypeFormat::FTMCDTimecode>());
 			// Set the first timecode sample on the track.
 			FTrackInfo::FFirstSampleTimecode firstTC;
 			firstTC.TimecodeValue = Electra::GetFromBigEndian(TimecodeBuffer[0]);
@@ -714,7 +729,7 @@ void FElectraProtronPlayer::FImpl::InternalOpen(const FString& InFilename)
 	}
 }
 
-void FElectraProtronPlayer::FImpl::GetTrackCodecInfo(ElectraProtronUtils::FCodecInfo& OutCodecInfo, const TSharedPtr<Electra::UtilitiesMP4::FMP4BoxTRAK, ESPMode::ThreadSafe>& InTrack, uint32 InTrackID)
+void FElectraProtronPlayer::FImpl::GetTrackCodecInfo(Electra::FCodecTypeFormat& OutCodecInfo, const TSharedPtr<Electra::UtilitiesMP4::FMP4BoxTRAK, ESPMode::ThreadSafe>& InTrack, uint32 InTrackID)
 {
 	// There needs to be an `stsd` box in this track. We do not try the expected path of `trak`->`mdia`->`minf`->`stbl`->`stsd` as
 	// this is not that much faster and if the file is somewhat ill-formed we may not find it if it's grouped under elsewhere.
@@ -758,21 +773,21 @@ void FElectraProtronPlayer::FImpl::GetTrackCodecInfo(ElectraProtronUtils::FCodec
 	// Based on the sample type, get it and see if it is using codec we support.
 	if (SampleType == Electra::UtilitiesMP4::FMP4BoxSampleEntry::ESampleType::Video)
 	{
-		auto SetResoAndFPSFromBox = [InTrack](ElectraProtronUtils::FCodecInfo::FVideo& InOutVideo, const TSharedPtr<Electra::UtilitiesMP4::FMP4BoxVisualSampleEntry, ESPMode::ThreadSafe>& InVisual) -> void
+		TArray<TSharedPtr<Electra::UtilitiesMP4::FMP4BoxVisualSampleEntry, ESPMode::ThreadSafe>> Visuals;
+		Stsd->GetSampleDescriptions(Visuals);
+		if (Visuals.Num() != 1)
 		{
-			InOutVideo.Width = InVisual->GetWidth();
-			InOutVideo.Height = InVisual->GetHeight();
-			auto Mdhd = InTrack->FindBoxRecursive<Electra::UtilitiesMP4::FMP4BoxMDHD>(Electra::UtilitiesMP4::MakeBoxAtom('m','d','h','d'), 2);
-			auto Stts = InTrack->FindBoxRecursive<Electra::UtilitiesMP4::FMP4BoxSTTS>(Electra::UtilitiesMP4::MakeBoxAtom('s','t','t','s'), 5);
-			if (Mdhd.IsValid() && Stts.IsValid() && Stts->GetEntries().Num())
+			if (Visuals.IsEmpty())
+			{
+				UE_LOG(LogElectraProtron, Warning, TEXT("No visual sample entry found in `stsd` box, ignoring track #%u."), InTrackID);
+			}
+			else
 			{
-				InOutVideo.FrameRate.SetFromND(Mdhd->GetTimescale(), Stts->GetEntries()[0].sample_delta);
+				UE_LOG(LogElectraProtron, Warning, TEXT("More than one visual sample entry found in `stsd` box, ignoring track #%u."), InTrackID);
 			}
-		};
+			return;
+		}
 
-		TArray<TSharedPtr<Electra::UtilitiesMP4::FMP4BoxVisualSampleEntry, ESPMode::ThreadSafe>> Visuals;
-		Stsd->GetSampleDescriptions(Visuals);
-		check(Visuals.Num() == 1);
 		switch(Visuals[0]->GetType())
 		{
 			case Electra::UtilitiesMP4::MakeBoxAtom('e','n','c','v'):
@@ -780,155 +795,30 @@ void FElectraProtronPlayer::FImpl::GetTrackCodecInfo(ElectraProtronUtils::FCodec
 				bIsEncrypted = true;
 				break;
 			}
-			case Electra::UtilitiesMP4::MakeBoxAtom('a','v','c','1'):
-			case Electra::UtilitiesMP4::MakeBoxAtom('a','v','c','3'):
-			{
-				if (Visuals[0]->GetFrameCount() == 1)
-				{
-					auto Avcc = Visuals[0]->FindBoxRecursive<Electra::UtilitiesMP4::FMP4BoxAVCC>(Electra::UtilitiesMP4::MakeBoxAtom('a','v','c','C'), 0);
-					if (Avcc.IsValid())
-					{
-						ElectraDecodersUtil::MPEG::H264::FAVCDecoderConfigurationRecord dcr;
-						if (dcr.Parse(Avcc->GetAVCDecoderConfigurationRecord()))
-						{
-							ElectraProtronUtils::FCodecInfo::FVideo Video;
-							ElectraDecodersUtil::MPEG::H264::FSequenceParameterSet sps;
-							const TCHAR* Prefix = nullptr;
-							if (dcr.GetSequenceParameterSets().Num() && ElectraDecodersUtil::MPEG::H264::ParseSequenceParameterSet(sps, dcr.GetSequenceParameterSets()[0]))
-							{
-								ElectraDecodersUtil::FFractionalValue fr = sps.GetTiming();
-								sps.GetDisplaySize(Video.Width, Video.Height);
-								if (fr.Num && fr.Denom)
-								{
-									Video.FrameRate.SetFromND(fr.Num, fr.Denom);
-								}
-								else
-								{
-									SetResoAndFPSFromBox(Video, Visuals[0]);
-								}
-								Prefix = TEXT("avc1");
-							}
-							else
-							{
-								SetResoAndFPSFromBox(Video, Visuals[0]);
-								Prefix = TEXT("avc3");
-							}
-							OutCodecInfo.Properties.Emplace<ElectraProtronUtils::FCodecInfo::FVideo>(MoveTemp(Video));
-							OutCodecInfo.RFC6381 = dcr.GetCodecSpecifierRFC6381(Prefix);
-							OutCodecInfo.FourCC = Visuals[0]->GetType();
-							OutCodecInfo.DCR = Avcc->GetAVCDecoderConfigurationRecord();
-							OutCodecInfo.CSD = dcr.GetCodecSpecificData();
-							OutCodecInfo.HumanReadableFormatInfo = dcr.GetFormatInfo();
-							OutCodecInfo.Type = ElectraProtronUtils::FCodecInfo::EType::Video;
-							AddChildren(Visuals[0]);
-						}
-					}
-				}
-				else
-				{
-					UE_LOG(LogElectraProtron, Warning, TEXT("Track #%u has a frame count other than 1 in the VisualSampleEntry, ignoring this track."), InTrackID);
-				}
-				break;
-			}
-			case Electra::UtilitiesMP4::MakeBoxAtom('h','v','c','1'):
-			case Electra::UtilitiesMP4::MakeBoxAtom('h','e','v','1'):
+			default:
 			{
 				if (Visuals[0]->GetFrameCount() == 1)
 				{
-					auto Hvcc = Visuals[0]->FindBoxRecursive<Electra::UtilitiesMP4::FMP4BoxHVCC>(Electra::UtilitiesMP4::MakeBoxAtom('h','v','c','C'), 0);
-					if (Hvcc.IsValid())
+					OutCodecInfo.Type = Electra::FCodecTypeFormat::EType::Video;
+					OutCodecInfo.FourCC = Visuals[0]->GetType();
+					OutCodecInfo.Properties.Emplace<Electra::FCodecTypeFormat::FVideo>();
+					Electra::FCodecTypeFormat::FVideo& cvid = OutCodecInfo.Properties.Get<Electra::FCodecTypeFormat::FVideo>();
+					cvid.Width = Visuals[0]->GetWidth();
+					cvid.Height = Visuals[0]->GetHeight();
+					auto Mdhd = InTrack->FindBoxRecursive<Electra::UtilitiesMP4::FMP4BoxMDHD>(Electra::UtilitiesMP4::MakeBoxAtom('m','d','h','d'), 2);
+					auto Stts = InTrack->FindBoxRecursive<Electra::UtilitiesMP4::FMP4BoxSTTS>(Electra::UtilitiesMP4::MakeBoxAtom('s','t','t','s'), 5);
+					if (Mdhd.IsValid() && Stts.IsValid() && Stts->GetEntries().Num())
 					{
-						ElectraDecodersUtil::MPEG::H265::FHEVCDecoderConfigurationRecord dcr;
-						if (dcr.Parse(Hvcc->GetHEVCDecoderConfigurationRecord()))
-						{
-							ElectraProtronUtils::FCodecInfo::FVideo Video;
-							ElectraDecodersUtil::MPEG::H265::FSequenceParameterSet sps;
-							const TCHAR* Prefix = nullptr;
-							if (dcr.GetSequenceParameterSets().Num() && ElectraDecodersUtil::MPEG::H265::ParseSequenceParameterSet(sps, dcr.GetSequenceParameterSets()[0]))
-							{
-								ElectraDecodersUtil::FFractionalValue fr = sps.GetTiming();
-								sps.GetDisplaySize(Video.Width, Video.Height);
-								Video.FrameRate.SetFromND(fr.Num, fr.Denom);
-								Prefix = TEXT("hvc1");
-							}
-							else
-							{
-								SetResoAndFPSFromBox(Video, Visuals[0]);
-								Prefix = TEXT("hev1");
-							}
-							OutCodecInfo.Properties.Emplace<ElectraProtronUtils::FCodecInfo::FVideo>(MoveTemp(Video));
-							OutCodecInfo.RFC6381 = dcr.GetCodecSpecifierRFC6381(Prefix);
-							OutCodecInfo.FourCC = Visuals[0]->GetType();
-							OutCodecInfo.DCR = Hvcc->GetHEVCDecoderConfigurationRecord();
-							OutCodecInfo.CSD = dcr.GetCodecSpecificData();
-							OutCodecInfo.HumanReadableFormatInfo = dcr.GetFormatInfo();
-							OutCodecInfo.Type = ElectraProtronUtils::FCodecInfo::EType::Video;
-							AddChildren(Visuals[0]);
-						}
+						cvid.FrameRate = FFrameRate(Mdhd->GetTimescale(), Stts->GetEntries()[0].sample_delta);
 					}
-				}
-				else
-				{
-					UE_LOG(LogElectraProtron, Warning, TEXT("Track #%u has a frame count other than 1 in the VisualSampleEntry, ignoring this track."), InTrackID);
-				}
-				break;
-			}
-			case Electra::UtilitiesMP4::MakeBoxAtom('v','p','0','8'):
-			case Electra::UtilitiesMP4::MakeBoxAtom('v','p','0','9'):
-			{
-				int32 VPx = Visuals[0]->GetType() == Electra::UtilitiesMP4::MakeBoxAtom('v','p','0','8') ? 8 : 9;
-				if (Visuals[0]->GetFrameCount() == 1)
-				{
-					auto Vpcc = Visuals[0]->FindBoxRecursive<Electra::UtilitiesMP4::FMP4BoxVPCC>(Electra::UtilitiesMP4::MakeBoxAtom('v','p','c','C'), 0);
-					if (Vpcc.IsValid())
+					for(auto& Ch : Visuals[0]->GetChildren())
 					{
-						ElectraDecodersUtil::VPxVideo::FVPCodecConfigurationRecord dcr;
-						if (dcr.Parse(VPx, Vpcc->GetVPCodecConfigurationRecord()))
-						{
-							ElectraProtronUtils::FCodecInfo::FVideo Video;
-
-							SetResoAndFPSFromBox(Video, Visuals[0]);
-
-							OutCodecInfo.Properties.Emplace<ElectraProtronUtils::FCodecInfo::FVideo>(MoveTemp(Video));
-							OutCodecInfo.RFC6381 = dcr.GetCodecSpecifierRFC6381();
-							OutCodecInfo.FourCC = Visuals[0]->GetType();
-							OutCodecInfo.DCR = Vpcc->GetVPCodecConfigurationRecord();
-							OutCodecInfo.CSD = dcr.GetCodecSpecificData();
-							OutCodecInfo.HumanReadableFormatInfo = dcr.GetFormatInfo();
-							OutCodecInfo.Type = ElectraProtronUtils::FCodecInfo::EType::Video;
-							AddChildren(Visuals[0]);
-						}
+						OutCodecInfo.ExtraBoxes.Emplace(Ch->GetType(), Ch->GetBoxData());
 					}
-				}
-				else
-				{
-					UE_LOG(LogElectraProtron, Warning, TEXT("Track #%u has a frame count other than 1 in the VisualSampleEntry, ignoring this track."), InTrackID);
-				}
-				break;
-			}
-			case Electra::UtilitiesMP4::MakeBoxAtom('a','v','0','1'):
-			{
-				if (Visuals[0]->GetFrameCount() == 1)
-				{
-					auto Av1c = Visuals[0]->FindBoxRecursive<Electra::UtilitiesMP4::FMP4BoxAV1C>(Electra::UtilitiesMP4::MakeBoxAtom('a','v','1','C'), 0);
-					if (Av1c.IsValid())
+					if (!CodecFormatUtils->PrepareCodecTypeFormat(OutCodecInfo, IElectraCodecFormatUtilsModularFeature::FCodecTypeFormatParams()))
 					{
-						ElectraDecodersUtil::AV1Video::FAV1CodecConfigurationRecord dcr;
-						if (dcr.Parse(Av1c->GetAV1CodecConfigurationRecord()))
-						{
-							ElectraProtronUtils::FCodecInfo::FVideo Video;
-
-							SetResoAndFPSFromBox(Video, Visuals[0]);
-
-							OutCodecInfo.Properties.Emplace<ElectraProtronUtils::FCodecInfo::FVideo>(MoveTemp(Video));
-							OutCodecInfo.RFC6381 = dcr.GetCodecSpecifierRFC6381();
-							OutCodecInfo.FourCC = Visuals[0]->GetType();
-							OutCodecInfo.DCR = Av1c->GetAV1CodecConfigurationRecord();
-							OutCodecInfo.CSD = dcr.GetCodecSpecificData();
-							OutCodecInfo.HumanReadableFormatInfo = dcr.GetFormatInfo();
-							OutCodecInfo.Type = ElectraProtronUtils::FCodecInfo::EType::Video;
-							AddChildren(Visuals[0]);
-						}
+						UE_LOG(LogElectraProtron, Warning, TEXT("Failed to prepare the codec format type, ignoring track #%u."), InTrackID);
+						return;
 					}
 				}
 				else
@@ -937,25 +827,24 @@ void FElectraProtronPlayer::FImpl::GetTrackCodecInfo(ElectraProtronUtils::FCodec
 				}
 				break;
 			}
-			default:
-			{
-				ElectraProtronUtils::FCodecInfo::FVideo Video;
-				SetResoAndFPSFromBox(Video, Visuals[0]);
-				OutCodecInfo.Properties.Emplace<ElectraProtronUtils::FCodecInfo::FVideo>(MoveTemp(Video));
-				OutCodecInfo.RFC6381 = Electra::UtilitiesMP4::GetPrintableBoxAtom(Visuals[0]->GetType());
-				OutCodecInfo.FourCC = Visuals[0]->GetType();
-				OutCodecInfo.DCR = Visuals[0]->GetBoxData();
-				OutCodecInfo.Type = ElectraProtronUtils::FCodecInfo::EType::Video;
-				AddChildren(Visuals[0]);
-				break;
-			}
 		}
 	}
 	else if (SampleType == Electra::UtilitiesMP4::FMP4BoxSampleEntry::ESampleType::Audio)
 	{
 		TArray<TSharedPtr<Electra::UtilitiesMP4::FMP4BoxAudioSampleEntry, ESPMode::ThreadSafe>> Audios;
 		Stsd->GetSampleDescriptions(Audios);
-		check(Audios.Num() == 1);
+		if (Audios.Num() != 1)
+		{
+			if (Audios.IsEmpty())
+			{
+				UE_LOG(LogElectraProtron, Warning, TEXT("No audio sample entry found in `stsd` box, ignoring track #%u."), InTrackID);
+			}
+			else
+			{
+				UE_LOG(LogElectraProtron, Warning, TEXT("More than one audio sample entry found in `stsd` box, ignoring track #%u."), InTrackID);
+			}
+			return;
+		}
 
 		switch(Audios[0]->GetType())
 		{
@@ -964,63 +853,23 @@ void FElectraProtronPlayer::FImpl::GetTrackCodecInfo(ElectraProtronUtils::FCodec
 				bIsEncrypted = true;
 				break;
 			}
-			case Electra::UtilitiesMP4::MakeBoxAtom('m','p','4','a'):
-			{
-				// Search down one extra level since the `esds` might be contained within a `wave` in QuickTime.
-				auto Esds = Audios[0]->FindBoxRecursive<Electra::UtilitiesMP4::FMP4BoxESDS>(Electra::UtilitiesMP4::MakeBoxAtom('e','s','d','s'), 1);
-				ElectraDecodersUtil::MPEG::FESDescriptor esds;
-				if (Esds.IsValid() && esds.Parse(Esds->GetESDescriptor()))
-				{
-					// AAC audio?
-					if (esds.GetObjectTypeID() == ElectraDecodersUtil::MPEG::FESDescriptor::FObjectTypeID::MPEG4_Audio &&
-						esds.GetStreamType() == ElectraDecodersUtil::MPEG::FESDescriptor::FStreamType::AudioStream)
-					{
-						ElectraDecodersUtil::MPEG::FAACDecoderConfigurationRecord dcr;
-						if (dcr.Parse(esds.GetCodecSpecificData()))
-						{
-							ElectraProtronUtils::FCodecInfo::FAudio Audio;
-							Audio.SampleRate = dcr.SamplingRate;
-							Audio.ChannelConfiguration = dcr.ChannelConfiguration;
-							Audio.NumChannels = ElectraDecodersUtil::MPEG::AACUtils::GetNumberOfChannelsFromChannelConfiguration(dcr.ChannelConfiguration);
-							OutCodecInfo.Properties.Emplace<ElectraProtronUtils::FCodecInfo::FAudio>(MoveTemp(Audio));
-							OutCodecInfo.RFC6381 = dcr.GetCodecSpecifierRFC6381();
-							OutCodecInfo.FourCC = Audios[0]->GetType();
-							OutCodecInfo.DCR = esds.GetCodecSpecificData();
-							OutCodecInfo.CSD = dcr.GetCodecSpecificData();
-							OutCodecInfo.HumanReadableFormatInfo = dcr.GetFormatInfo();
-							OutCodecInfo.Type = ElectraProtronUtils::FCodecInfo::EType::Audio;
-							AddChildren(Audios[0]);
-						}
-					}
-					// MPEG audio?
-					else if (esds.GetObjectTypeID() == ElectraDecodersUtil::MPEG::FESDescriptor::FObjectTypeID::MPEG1_Audio &&
-							 esds.GetStreamType() == ElectraDecodersUtil::MPEG::FESDescriptor::FStreamType::AudioStream)
-					{
-						ElectraProtronUtils::FCodecInfo::FAudio Audio;
-						Audio.SampleRate = Audios[0]->GetSampleRate();
-						Audio.NumChannels = Audios[0]->GetChannelCount();
-						OutCodecInfo.Properties.Emplace<ElectraProtronUtils::FCodecInfo::FAudio>(MoveTemp(Audio));
-						OutCodecInfo.RFC6381 = TEXT("mp4a.40.34");
-						OutCodecInfo.FourCC = Electra::UtilitiesMP4::MakeBoxAtom('m','p','g','a');
-						OutCodecInfo.HumanReadableFormatInfo = TEXT("MPEG audio");
-						OutCodecInfo.Type = ElectraProtronUtils::FCodecInfo::EType::Audio;
-						AddChildren(Audios[0]);
-					}
-				}
-				break;
-			}
 			default:
 			{
-				ElectraProtronUtils::FCodecInfo::FAudio Audio;
-				Audio.SampleRate = Audios[0]->GetSampleRate();
-				Audio.ChannelConfiguration = 0;
-				Audio.NumChannels = Audios[0]->GetChannelCount();
-				OutCodecInfo.Properties.Emplace<ElectraProtronUtils::FCodecInfo::FAudio>(MoveTemp(Audio));
-				OutCodecInfo.RFC6381 = Electra::UtilitiesMP4::GetPrintableBoxAtom(Audios[0]->GetType());
+				OutCodecInfo.Type = Electra::FCodecTypeFormat::EType::Audio;
 				OutCodecInfo.FourCC = Audios[0]->GetType();
-				OutCodecInfo.DCR = Audios[0]->GetBoxData();
-				OutCodecInfo.Type = ElectraProtronUtils::FCodecInfo::EType::Audio;
-				AddChildren(Audios[0]);
+				OutCodecInfo.Properties.Emplace<Electra::FCodecTypeFormat::FAudio>();
+				Electra::FCodecTypeFormat::FAudio& caud = OutCodecInfo.Properties.Get<Electra::FCodecTypeFormat::FAudio>();
+				caud.NumChannels = Audios[0]->GetChannelCount();
+				caud.SampleRate = Audios[0]->GetSampleRate();
+				for(auto& Ch : Audios[0]->GetChildren())
+				{
+					OutCodecInfo.ExtraBoxes.Emplace(Ch->GetType(), Ch->GetBoxData());
+				}
+				if (!CodecFormatUtils->PrepareCodecTypeFormat(OutCodecInfo, IElectraCodecFormatUtilsModularFeature::FCodecTypeFormatParams()))
+				{
+					UE_LOG(LogElectraProtron, Warning, TEXT("Failed to prepare the codec format type, ignoring track #%u."), InTrackID);
+					return;
+				}
 				break;
 			}
 		}
@@ -1035,17 +884,17 @@ void FElectraProtronPlayer::FImpl::GetTrackCodecInfo(ElectraProtronUtils::FCodec
 		{
 			case Electra::UtilitiesMP4::MakeBoxAtom('t','m','c','d'):
 			{
-				ElectraProtronUtils::FCodecInfo::FTMCDTimecode Timecode;
+				Electra::FCodecTypeFormat::FTMCDTimecode Timecode;
 
 				Timecode.Flags = Timecodes[0]->GetFlags();
 				Timecode.Timescale = Timecodes[0]->GetTimescale();
 				Timecode.FrameDuration = Timecodes[0]->GetFrameDuration();
 				Timecode.NumberOfFrames = Timecodes[0]->GetNumberOfFrames();
-				OutCodecInfo.Properties.Emplace<ElectraProtronUtils::FCodecInfo::FTMCDTimecode>(MoveTemp(Timecode));
+				OutCodecInfo.Properties.Emplace<Electra::FCodecTypeFormat::FTMCDTimecode>(MoveTemp(Timecode));
 				OutCodecInfo.RFC6381 = Electra::UtilitiesMP4::GetPrintableBoxAtom(Timecodes[0]->GetType());
 				OutCodecInfo.FourCC = Timecodes[0]->GetType();
 				OutCodecInfo.DCR = Timecodes[0]->GetBoxData();
-				OutCodecInfo.Type = ElectraProtronUtils::FCodecInfo::EType::Timecode;
+				OutCodecInfo.Type = Electra::FCodecTypeFormat::EType::Timecode;
 				AddChildren(Timecodes[0]);
 				break;
 			}
@@ -1053,7 +902,7 @@ void FElectraProtronPlayer::FImpl::GetTrackCodecInfo(ElectraProtronUtils::FCodec
 	}
 
 	// If the track has not been identified as usable so far, remove it.
-	if (OutCodecInfo.Type == ElectraProtronUtils::FCodecInfo::EType::Invalid)
+	if (OutCodecInfo.Type == Electra::FCodecTypeFormat::EType::Invalid)
 	{
 		if (!bIsEncrypted)
 		{
@@ -1076,7 +925,7 @@ FVariant FElectraProtronPlayer::FImpl::GetMediaInfo(FName InInfoName)
 	if (InInfoName == ElectraProtronOptionNames::StartTimecodeValue || InInfoName == ElectraProtronOptionNames::StartTimecodeFrameRate ||
 		InInfoName == ElectraProtronOptionNames::KeyframeInterval || InInfoName == ElectraProtronOptionNames::VideoDecoderProvider)
 	{
-		auto ci = CodecTypeIndex(ElectraProtronUtils::FCodecInfo::EType::Video);
+		auto ci = CodecTypeIndex(Electra::FCodecTypeFormat::EType::Video);
 		int32 SelectedVideoTrackIndex = TrackSelection.SelectedTrackIndex[ci];
 		if (SelectedVideoTrackIndex >= 0 && SelectedVideoTrackIndex < UsableTrackArrayIndicesByType[ci].Num())
 		{
@@ -1102,7 +951,7 @@ FVariant FElectraProtronPlayer::FImpl::GetMediaInfo(FName InInfoName)
 	}
 	else if (InInfoName == ElectraProtronOptionNames::AudioDecoderProvider)
 	{
-		auto ci = CodecTypeIndex(ElectraProtronUtils::FCodecInfo::EType::Audio);
+		auto ci = CodecTypeIndex(Electra::FCodecTypeFormat::EType::Audio);
 		int32 SelectedAudioTrackIndex = TrackSelection.SelectedTrackIndex[ci];
 		if (SelectedAudioTrackIndex >= 0 && SelectedAudioTrackIndex < UsableTrackArrayIndicesByType[ci].Num())
 		{
@@ -1124,11 +973,11 @@ int32 FElectraProtronPlayer::FImpl::GetNumTracks(EMediaTrackType InTrackType)
 	{
 		case EMediaTrackType::Video:
 		{
-			return UsableTrackArrayIndicesByType[CodecTypeIndex(ElectraProtronUtils::FCodecInfo::EType::Video)].Num();
+			return UsableTrackArrayIndicesByType[CodecTypeIndex(Electra::FCodecTypeFormat::EType::Video)].Num();
 		}
 		case EMediaTrackType::Audio:
 		{
-			return UsableTrackArrayIndicesByType[CodecTypeIndex(ElectraProtronUtils::FCodecInfo::EType::Audio)].Num();
+			return UsableTrackArrayIndicesByType[CodecTypeIndex(Electra::FCodecTypeFormat::EType::Audio)].Num();
 		}
 	}
 	return 0;
@@ -1142,11 +991,11 @@ int32 FElectraProtronPlayer::FImpl::GetNumTrackFormats(EMediaTrackType InTrackTy
 		{
 			case EMediaTrackType::Video:
 			{
-				return InTrackIndex < UsableTrackArrayIndicesByType[CodecTypeIndex(ElectraProtronUtils::FCodecInfo::EType::Video)].Num() ? 1 : 0;
+				return InTrackIndex < UsableTrackArrayIndicesByType[CodecTypeIndex(Electra::FCodecTypeFormat::EType::Video)].Num() ? 1 : 0;
 			}
 			case EMediaTrackType::Audio:
 			{
-				return InTrackIndex < UsableTrackArrayIndicesByType[CodecTypeIndex(ElectraProtronUtils::FCodecInfo::EType::Audio)].Num() ? 1 : 0;
+				return InTrackIndex < UsableTrackArrayIndicesByType[CodecTypeIndex(Electra::FCodecTypeFormat::EType::Audio)].Num() ? 1 : 0;
 			}
 		}
 	}
@@ -1162,11 +1011,11 @@ int32 FElectraProtronPlayer::FImpl::GetTrackFormat(EMediaTrackType InTrackType,
 		{
 			case EMediaTrackType::Video:
 			{
-				return InTrackIndex < UsableTrackArrayIndicesByType[CodecTypeIndex(ElectraProtronUtils::FCodecInfo::EType::Video)].Num() ? 0 : -1;
+				return InTrackIndex < UsableTrackArrayIndicesByType[CodecTypeIndex(Electra::FCodecTypeFormat::EType::Video)].Num() ? 0 : -1;
 			}
 			case EMediaTrackType::Audio:
 			{
-				return InTrackIndex < UsableTrackArrayIndicesByType[CodecTypeIndex(ElectraProtronUtils::FCodecInfo::EType::Audio)].Num() ? 0 : -1;
+				return InTrackIndex < UsableTrackArrayIndicesByType[CodecTypeIndex(Electra::FCodecTypeFormat::EType::Audio)].Num() ? 0 : -1;
 			}
 		}
 	}
@@ -1181,9 +1030,9 @@ FText FElectraProtronPlayer::FImpl::GetTrackDisplayName(EMediaTrackType InTrackT
 		{
 			case EMediaTrackType::Video:
 			{
-				if (InTrackIndex < UsableTrackArrayIndicesByType[CodecTypeIndex(ElectraProtronUtils::FCodecInfo::EType::Video)].Num())
+				if (InTrackIndex < UsableTrackArrayIndicesByType[CodecTypeIndex(Electra::FCodecTypeFormat::EType::Video)].Num())
 				{
-					const FTrackInfo& ti(*Tracks[UsableTrackArrayIndicesByType[CodecTypeIndex(ElectraProtronUtils::FCodecInfo::EType::Video)][InTrackIndex]]);
+					const FTrackInfo& ti(*Tracks[UsableTrackArrayIndicesByType[CodecTypeIndex(Electra::FCodecTypeFormat::EType::Video)][InTrackIndex]]);
 					FString Name(ti.MP4Track->GetCommonMetadata().Name);
 					if (Name.IsEmpty())
 					{
@@ -1195,9 +1044,9 @@ FText FElectraProtronPlayer::FImpl::GetTrackDisplayName(EMediaTrackType InTrackT
 			}
 			case EMediaTrackType::Audio:
 			{
-				if (InTrackIndex < UsableTrackArrayIndicesByType[CodecTypeIndex(ElectraProtronUtils::FCodecInfo::EType::Audio)].Num())
+				if (InTrackIndex < UsableTrackArrayIndicesByType[CodecTypeIndex(Electra::FCodecTypeFormat::EType::Audio)].Num())
 				{
-					const FTrackInfo& ti(*Tracks[UsableTrackArrayIndicesByType[CodecTypeIndex(ElectraProtronUtils::FCodecInfo::EType::Audio)][InTrackIndex]]);
+					const FTrackInfo& ti(*Tracks[UsableTrackArrayIndicesByType[CodecTypeIndex(Electra::FCodecTypeFormat::EType::Audio)][InTrackIndex]]);
 					FString Name(ti.MP4Track->GetCommonMetadata().Name);
 					if (Name.IsEmpty())
 					{
@@ -1220,17 +1069,17 @@ FString FElectraProtronPlayer::FImpl::GetTrackLanguage(EMediaTrackType InTrackTy
 		{
 			case EMediaTrackType::Video:
 			{
-				if (InTrackIndex < UsableTrackArrayIndicesByType[CodecTypeIndex(ElectraProtronUtils::FCodecInfo::EType::Video)].Num())
+				if (InTrackIndex < UsableTrackArrayIndicesByType[CodecTypeIndex(Electra::FCodecTypeFormat::EType::Video)].Num())
 				{
-					return Tracks[UsableTrackArrayIndicesByType[CodecTypeIndex(ElectraProtronUtils::FCodecInfo::EType::Video)][InTrackIndex]]->MP4Track->GetCommonMetadata().LanguageTag.Get();
+					return Tracks[UsableTrackArrayIndicesByType[CodecTypeIndex(Electra::FCodecTypeFormat::EType::Video)][InTrackIndex]]->MP4Track->GetCommonMetadata().LanguageTag.Get();
 				}
 				break;
 			}
 			case EMediaTrackType::Audio:
 			{
-				if (InTrackIndex < UsableTrackArrayIndicesByType[CodecTypeIndex(ElectraProtronUtils::FCodecInfo::EType::Audio)].Num())
+				if (InTrackIndex < UsableTrackArrayIndicesByType[CodecTypeIndex(Electra::FCodecTypeFormat::EType::Audio)].Num())
 				{
-					return Tracks[UsableTrackArrayIndicesByType[CodecTypeIndex(ElectraProtronUtils::FCodecInfo::EType::Audio)][InTrackIndex]]->MP4Track->GetCommonMetadata().LanguageTag.Get();
+					return Tracks[UsableTrackArrayIndicesByType[CodecTypeIndex(Electra::FCodecTypeFormat::EType::Audio)][InTrackIndex]]->MP4Track->GetCommonMetadata().LanguageTag.Get();
 				}
 				break;
 			}
@@ -1247,17 +1096,17 @@ FString FElectraProtronPlayer::FImpl::GetTrackName(EMediaTrackType InTrackType,
 		{
 			case EMediaTrackType::Video:
 			{
-				if (InTrackIndex < UsableTrackArrayIndicesByType[CodecTypeIndex(ElectraProtronUtils::FCodecInfo::EType::Video)].Num())
+				if (InTrackIndex < UsableTrackArrayIndicesByType[CodecTypeIndex(Electra::FCodecTypeFormat::EType::Video)].Num())
 				{
-					return FString::Printf(TEXT("%u"), Tracks[UsableTrackArrayIndicesByType[CodecTypeIndex(ElectraProtronUtils::FCodecInfo::EType::Video)][InTrackIndex]]->TrackID);
+					return FString::Printf(TEXT("%u"), Tracks[UsableTrackArrayIndicesByType[CodecTypeIndex(Electra::FCodecTypeFormat::EType::Video)][InTrackIndex]]->TrackID);
 				}
 				break;
 			}
 			case EMediaTrackType::Audio:
 			{
-				if (InTrackIndex < UsableTrackArrayIndicesByType[CodecTypeIndex(ElectraProtronUtils::FCodecInfo::EType::Audio)].Num())
+				if (InTrackIndex < UsableTrackArrayIndicesByType[CodecTypeIndex(Electra::FCodecTypeFormat::EType::Audio)].Num())
 				{
-					return FString::Printf(TEXT("%u"), Tracks[UsableTrackArrayIndicesByType[CodecTypeIndex(ElectraProtronUtils::FCodecInfo::EType::Audio)][InTrackIndex]]->TrackID);
+					return FString::Printf(TEXT("%u"), Tracks[UsableTrackArrayIndicesByType[CodecTypeIndex(Electra::FCodecTypeFormat::EType::Audio)][InTrackIndex]]->TrackID);
 				}
 				break;
 			}
@@ -1268,14 +1117,22 @@ FString FElectraProtronPlayer::FImpl::GetTrackName(EMediaTrackType InTrackType,
 
 bool FElectraProtronPlayer::FImpl::GetVideoTrackFormat(int32 InTrackIndex, int32 InFormatIndex, FMediaVideoTrackFormat& OutFormat)
 {
-	if (InTrackIndex >= 0 && InTrackIndex < UsableTrackArrayIndicesByType[CodecTypeIndex(ElectraProtronUtils::FCodecInfo::EType::Video)].Num() && InFormatIndex == 0)
+	if (InTrackIndex >= 0 && InTrackIndex < UsableTrackArrayIndicesByType[CodecTypeIndex(Electra::FCodecTypeFormat::EType::Video)].Num() && InFormatIndex == 0)
 	{
-		const FTrackInfo& ti(*Tracks[UsableTrackArrayIndicesByType[CodecTypeIndex(ElectraProtronUtils::FCodecInfo::EType::Video)][InTrackIndex]]);
-		const ElectraProtronUtils::FCodecInfo::FVideo& vi(ti.CodecInfo.Properties.Get<ElectraProtronUtils::FCodecInfo::FVideo>());
+		const FTrackInfo& ti(*Tracks[UsableTrackArrayIndicesByType[CodecTypeIndex(Electra::FCodecTypeFormat::EType::Video)][InTrackIndex]]);
+		const Electra::FCodecTypeFormat::FVideo& vi(ti.CodecInfo.Properties.Get<Electra::FCodecTypeFormat::FVideo>());
 		OutFormat.Dim.X = (int32) vi.Width;
 		OutFormat.Dim.Y = (int32) vi.Height;
-		OutFormat.FrameRate = (float) vi.FrameRate.GetAsDouble();
-		OutFormat.FrameRates = TRange<float>{ OutFormat.FrameRate };
+		if (vi.FrameRate.IsValid())
+		{
+			OutFormat.FrameRate = (float) vi.FrameRate.AsDecimal();
+			OutFormat.FrameRates = TRange<float>{ OutFormat.FrameRate };
+		}
+		else
+		{
+			OutFormat.FrameRate = 0.0f;
+			OutFormat.FrameRates = TRange<float>{ OutFormat.FrameRate };//TRange<float>::Empty();
+		}
 		OutFormat.TypeName = ti.HumanReadableCodecFormat;
 		return true;
 	}
@@ -1284,10 +1141,10 @@ bool FElectraProtronPlayer::FImpl::GetVideoTrackFormat(int32 InTrackIndex, int32
 
 bool FElectraProtronPlayer::FImpl::GetAudioTrackFormat(int32 InTrackIndex, int32 InFormatIndex, FMediaAudioTrackFormat& OutFormat)
 {
-	if (InTrackIndex >= 0 && InTrackIndex < UsableTrackArrayIndicesByType[CodecTypeIndex(ElectraProtronUtils::FCodecInfo::EType::Audio)].Num() && InFormatIndex == 0)
+	if (InTrackIndex >= 0 && InTrackIndex < UsableTrackArrayIndicesByType[CodecTypeIndex(Electra::FCodecTypeFormat::EType::Audio)].Num() && InFormatIndex == 0)
 	{
-		const FTrackInfo& ti(*Tracks[UsableTrackArrayIndicesByType[CodecTypeIndex(ElectraProtronUtils::FCodecInfo::EType::Audio)][InTrackIndex]]);
-		const ElectraProtronUtils::FCodecInfo::FAudio& ai(ti.CodecInfo.Properties.Get<ElectraProtronUtils::FCodecInfo::FAudio>());
+		const FTrackInfo& ti(*Tracks[UsableTrackArrayIndicesByType[CodecTypeIndex(Electra::FCodecTypeFormat::EType::Audio)][InTrackIndex]]);
+		const Electra::FCodecTypeFormat::FAudio& ai(ti.CodecInfo.Properties.Get<Electra::FCodecTypeFormat::FAudio>());
 		OutFormat.BitsPerSample = 16;
 		OutFormat.NumChannels = ai.NumChannels;
 		OutFormat.SampleRate = ai.SampleRate;
@@ -1303,11 +1160,11 @@ int32 FElectraProtronPlayer::FImpl::GetSelectedTrack(EMediaTrackType InTrackType
 	{
 		case EMediaTrackType::Video:
 		{
-			return TrackSelection.SelectedTrackIndex[CodecTypeIndex(ElectraProtronUtils::FCodecInfo::EType::Video)];
+			return TrackSelection.SelectedTrackIndex[CodecTypeIndex(Electra::FCodecTypeFormat::EType::Video)];
 		}
 		case EMediaTrackType::Audio:
 		{
-			return TrackSelection.SelectedTrackIndex[CodecTypeIndex(ElectraProtronUtils::FCodecInfo::EType::Audio)];
+			return TrackSelection.SelectedTrackIndex[CodecTypeIndex(Electra::FCodecTypeFormat::EType::Audio)];
 		}
 	}
 	return -1;
@@ -1324,7 +1181,7 @@ bool FElectraProtronPlayer::FImpl::SelectTrack(EMediaTrackType InTrackType, int3
 	{
 		case EMediaTrackType::Video:
 		{
-			const auto TypeIndex = CodecTypeIndex(ElectraProtronUtils::FCodecInfo::EType::Video);
+			const auto TypeIndex = CodecTypeIndex(Electra::FCodecTypeFormat::EType::Video);
 			if (InTrackIndex < UsableTrackArrayIndicesByType[TypeIndex].Num())
 			{
 				if (TrackSelection.SelectedTrackIndex[TypeIndex] != InTrackIndex)
@@ -1341,7 +1198,7 @@ bool FElectraProtronPlayer::FImpl::SelectTrack(EMediaTrackType InTrackType, int3
 		}
 		case EMediaTrackType::Audio:
 		{
-			const auto TypeIndex = CodecTypeIndex(ElectraProtronUtils::FCodecInfo::EType::Audio);
+			const auto TypeIndex = CodecTypeIndex(Electra::FCodecTypeFormat::EType::Audio);
 			if (InTrackIndex < UsableTrackArrayIndicesByType[TypeIndex].Num())
 			{
 				if (TrackSelection.SelectedTrackIndex[TypeIndex] != InTrackIndex)
@@ -1373,12 +1230,12 @@ bool FElectraProtronPlayer::FImpl::QueryCacheState(EMediaCacheState InState, TRa
 {
 	if (InState == EMediaCacheState::Loading)
 	{
-		if (TrackSelection.SelectedTrackIndex[CodecTypeIndex(ElectraProtronUtils::FCodecInfo::EType::Video)] >= 0)
+		if (TrackSelection.SelectedTrackIndex[CodecTypeIndex(Electra::FCodecTypeFormat::EType::Video)] >= 0)
 		{
 			OutTimeRanges = VideoLoaderThread.GetTimeRangesToLoad();
 			return true;
 		}
-		else if (TrackSelection.SelectedTrackIndex[CodecTypeIndex(ElectraProtronUtils::FCodecInfo::EType::Audio)] >= 0)
+		else if (TrackSelection.SelectedTrackIndex[CodecTypeIndex(Electra::FCodecTypeFormat::EType::Audio)] >= 0)
 		{
 			OutTimeRanges = AudioLoaderThread.GetTimeRangesToLoad();
 			return true;
@@ -1386,7 +1243,7 @@ bool FElectraProtronPlayer::FImpl::QueryCacheState(EMediaCacheState InState, TRa
 	}
 	else if (InState == EMediaCacheState::Loaded)
 	{
-		if (TrackSelection.SelectedTrackIndex[CodecTypeIndex(ElectraProtronUtils::FCodecInfo::EType::Video)] >= 0)
+		if (TrackSelection.SelectedTrackIndex[CodecTypeIndex(Electra::FCodecTypeFormat::EType::Video)] >= 0)
 		{
 			TSharedPtr<FSampleQueueInterface, ESPMode::ThreadSafe> sqi(GetCurrentSampleQueueInterface());
 			if (sqi.IsValid())
@@ -1530,8 +1387,8 @@ void FElectraProtronPlayer::FImpl::TickInput(FTimespan InDeltaTime, FTimespan In
 	if (sqi.IsValid())
 	{
 		FMediaTimeStamp ts = sqi->GetLastHandedOutTimestamp();
-		const bool bIsVideoActive = TrackSelection.ActiveTrackIndex[CodecTypeIndex(ElectraProtronUtils::FCodecInfo::EType::Video)] != -1;
-		const bool bIsAudioActive = TrackSelection.ActiveTrackIndex[CodecTypeIndex(ElectraProtronUtils::FCodecInfo::EType::Audio)] != -1;
+		const bool bIsVideoActive = TrackSelection.ActiveTrackIndex[CodecTypeIndex(Electra::FCodecTypeFormat::EType::Video)] != -1;
+		const bool bIsAudioActive = TrackSelection.ActiveTrackIndex[CodecTypeIndex(Electra::FCodecTypeFormat::EType::Audio)] != -1;
 		if (ts.IsValid())
 		{
 			FTimespan NewPos = ts.GetTime();
@@ -1573,11 +1430,11 @@ void FElectraProtronPlayer::FImpl::UpdateTrackLoader(int32 InCodecTypeIndex)
 				TrackSampleBuffers.Emplace(Track->TrackID, MoveTemp(tsb));
 			}
 
-			if (InCodecTypeIndex == CodecTypeIndex(ElectraProtronUtils::FCodecInfo::EType::Video))
+			if (InCodecTypeIndex == CodecTypeIndex(Electra::FCodecTypeFormat::EType::Video))
 			{
 				VideoLoaderThread.RequestLoad(TrackSampleBuffers[Track->TrackID], CurrentPlayPosTime);
 			}
-			else if (InCodecTypeIndex == CodecTypeIndex(ElectraProtronUtils::FCodecInfo::EType::Audio))
+			else if (InCodecTypeIndex == CodecTypeIndex(Electra::FCodecTypeFormat::EType::Audio))
 			{
 				AudioLoaderThread.RequestLoad(TrackSampleBuffers[Track->TrackID], CurrentPlayPosTime);
 			}
diff --git a/Engine/Plugins/Media/ElectraPlayer/Source/ElectraProtron/Private/Player/ElectraProtronPlayerImpl.h b/Engine/Plugins/Media/ElectraPlayer/Source/ElectraProtron/Private/Player/ElectraProtronPlayerImpl.h
--- a/Engine/Plugins/Media/ElectraPlayer/Source/ElectraProtron/Private/Player/ElectraProtronPlayerImpl.h
+++ b/Engine/Plugins/Media/ElectraPlayer/Source/ElectraProtron/Private/Player/ElectraProtronPlayerImpl.h
@@ -14,8 +14,9 @@
 #include "Utilities/UtilitiesMP4.h"
 #include "Utilities/MP4Boxes/MP4Boxes.h"
 #include "Utilities/MP4Boxes/MP4Track.h"
-#include "TrackFormatInfo.h"
+#include "CodecTypeFormat.h"
 #include "IElectraDecoder.h"
+#include "IElectraCodecFormatUtils.h"
 #include "ElectraDecodersPlatformResources.h"
 #include "MediaSamples.h"
 #include "ElectraTextureSample.h"
@@ -288,7 +289,7 @@ public:
 	int32 NumVideoSamples() const;
 
 private:
-	static constexpr int32 CodecTypeIndex(const ElectraProtronUtils::FCodecInfo::EType InType)
+	static constexpr int32 CodecTypeIndex(const Electra::FCodecTypeFormat::EType InType)
 	{ return (int32) InType; }
 
 	struct FConfig
@@ -339,7 +340,7 @@ private:
 	struct FTrackInfo
 	{
 		TArray<TWeakPtr<FTrackInfo, ESPMode::ThreadSafe>> IsReferencedByTracks;
-		ElectraProtronUtils::FCodecInfo CodecInfo;
+		Electra::FCodecTypeFormat CodecInfo;
 		FString HumanReadableCodecFormat;
 		struct FDecoderProviderInfo
 		{
@@ -412,7 +413,7 @@ private:
 	void SendWorkerThreadMessage(FWorkerThreadMessage&& InMessage);
 
 	void InternalOpen(const FString& InFilename);
-	void GetTrackCodecInfo(ElectraProtronUtils::FCodecInfo& OutCodecInfo, const TSharedPtr<Electra::UtilitiesMP4::FMP4BoxTRAK, ESPMode::ThreadSafe>& InTrack, uint32 InTrackID);
+	void GetTrackCodecInfo(Electra::FCodecTypeFormat& OutCodecInfo, const TSharedPtr<Electra::UtilitiesMP4::FMP4BoxTRAK, ESPMode::ThreadSafe>& InTrack, uint32 InTrackID);
 
 	void UpdateTrackLoader(int32 InCodecTypeIndex);
 
@@ -420,16 +421,18 @@ private:
 	void HandleRateChanges();
 	void HandleSeekRequest(const FSeekRequest& InSeekRequest);
 
-	const int32 kCodecTrackIndexMap[(int32) ElectraProtronUtils::FCodecInfo::EType::MAX] =
+	const int32 kCodecTrackIndexMap[(int32) Electra::FCodecTypeFormat::EType::MAX] =
 	{
-		CodecTypeIndex(ElectraProtronUtils::FCodecInfo::EType::Video),
-		CodecTypeIndex(ElectraProtronUtils::FCodecInfo::EType::Audio),
-		CodecTypeIndex(ElectraProtronUtils::FCodecInfo::EType::Subtitle),
-		CodecTypeIndex(ElectraProtronUtils::FCodecInfo::EType::Timecode)
+		CodecTypeIndex(Electra::FCodecTypeFormat::EType::Video),
+		CodecTypeIndex(Electra::FCodecTypeFormat::EType::Audio),
+		CodecTypeIndex(Electra::FCodecTypeFormat::EType::Subtitle),
+		CodecTypeIndex(Electra::FCodecTypeFormat::EType::Timecode)
 	};
 
 	FConfig Config;
 
+	IElectraCodecFormatUtilsModularFeature* CodecFormatUtils = nullptr;
+
 	FRunnableThread* Thread = nullptr;
 	ImplPointer SelfDuringTerminate;
 	FMediaEvent WorkMessageSignal;
@@ -540,7 +543,7 @@ private:
 		};
 
 		ELoadResult Load(FMP4TrackSampleBufferPtr InTrackSampleBuffer, FTrackIterator InAtIterator);
-		FMP4SamplePtr RetrieveSample(const FMP4TrackSampleBufferPtr& InTrackSampleBuffer, const FTrackIterator& InSampleIt, const FTrackIterator& InOptionalTimecodeIt, const ElectraProtronUtils::FCodecInfo::FTMCDTimecode& InTimecodeInfo);
+		FMP4SamplePtr RetrieveSample(const FMP4TrackSampleBufferPtr& InTrackSampleBuffer, const FTrackIterator& InSampleIt, const FTrackIterator& InOptionalTimecodeIt, const Electra::FCodecTypeFormat::FTMCDTimecode& InTimecodeInfo);
 
 		struct FSampleRange
 		{
@@ -715,9 +718,9 @@ private:
 		TSharedPtr<FElectraTextureSamplePool, ESPMode::ThreadSafe> VideoDecoderOutputPool;
 	};
 
-	FLoaderThread VideoLoaderThread {Config, CodecTypeIndex(ElectraProtronUtils::FCodecInfo::EType::Video)};
-	FLoaderThread AudioLoaderThread {Config, CodecTypeIndex(ElectraProtronUtils::FCodecInfo::EType::Audio)};
+	FLoaderThread VideoLoaderThread {Config, CodecTypeIndex(Electra::FCodecTypeFormat::EType::Video)};
+	FLoaderThread AudioLoaderThread {Config, CodecTypeIndex(Electra::FCodecTypeFormat::EType::Audio)};
 
-	FDecoderThread VideoDecoderThread {Config, CodecTypeIndex(ElectraProtronUtils::FCodecInfo::EType::Video)};
-	FDecoderThread AudioDecoderThread {Config, CodecTypeIndex(ElectraProtronUtils::FCodecInfo::EType::Audio)};
+	FDecoderThread VideoDecoderThread {Config, CodecTypeIndex(Electra::FCodecTypeFormat::EType::Video)};
+	FDecoderThread AudioDecoderThread {Config, CodecTypeIndex(Electra::FCodecTypeFormat::EType::Audio)};
 };
diff --git a/Engine/Plugins/Media/ElectraPlayer/Source/ElectraProtron/Private/Player/ElectraProtronPlayerLoader.cpp b/Engine/Plugins/Media/ElectraPlayer/Source/ElectraProtron/Private/Player/ElectraProtronPlayerLoader.cpp
--- a/Engine/Plugins/Media/ElectraPlayer/Source/ElectraProtron/Private/Player/ElectraProtronPlayerLoader.cpp
+++ b/Engine/Plugins/Media/ElectraPlayer/Source/ElectraProtron/Private/Player/ElectraProtronPlayerLoader.cpp
@@ -5,10 +5,7 @@
 #include "IElectraCodecFactoryModule.h"
 #include "IElectraCodecFactory.h"
 #include "IElectraDecoderFeaturesAndOptions.h"
-#include "Utils/MPEG/ElectraUtilsMPEGVideo_H264.h"
-#include "Utils/MPEG/ElectraUtilsMPEGVideo_H265.h"
-#include "Utils/MPEG/ElectraUtilsMPEGAudio.h"
-#include "TrackFormatInfo.h"
+#include "CodecTypeFormat.h"
 
 
 
@@ -337,8 +334,8 @@ FElectraProtronPlayer::FImpl::FLoaderThread::ELoadResult FElectraProtronPlayer::
 
 	// Are timecodes from another track referenced?
 	FTrackIterator FwdTkItTC, RevTkItTC;
-	ElectraProtronUtils::FCodecInfo::FTMCDTimecode TimecodeInfo;
-	if (Config.bReadSampleTimecode && LoaderTypeIndex == CodecTypeIndex(ElectraProtronUtils::FCodecInfo::EType::Video))
+	Electra::FCodecTypeFormat::FTMCDTimecode TimecodeInfo;
+	if (Config.bReadSampleTimecode && LoaderTypeIndex == CodecTypeIndex(Electra::FCodecTypeFormat::EType::Video))
 	{
 		if (auto TimecodeTrack = InTrackSampleBuffer->TrackAndCodecInfo->ReferencedTimecodeTrack.Pin())
 		{
@@ -347,7 +344,7 @@ FElectraProtronPlayer::FImpl::FLoaderThread::ELoadResult FElectraProtronPlayer::
 			if (TimecodeTrack->MP4Track->GetNumberOfSamples() == InTrackSampleBuffer->TrackAndCodecInfo->MP4Track->GetNumberOfSamples())
 			{
 				// Get the timecode description from the codec info.
-				TimecodeInfo = TimecodeTrack->CodecInfo.Properties.Get<ElectraProtronUtils::FCodecInfo::FTMCDTimecode>();
+				TimecodeInfo = TimecodeTrack->CodecInfo.Properties.Get<Electra::FCodecTypeFormat::FTMCDTimecode>();
 				FwdTkItTC = TimecodeTrack->MP4Track->CreateIterator(InAtIterator->GetSampleNumber());
 				if (FwdTkItTC.IsValid())
 				{
@@ -457,7 +454,7 @@ FElectraProtronPlayer::FImpl::FLoaderThread::ELoadResult FElectraProtronPlayer::
 	return ELoadResult::Ok;
 }
 
-TSharedPtr<FElectraProtronPlayer::FImpl::FMP4Sample, ESPMode::ThreadSafe> FElectraProtronPlayer::FImpl::FLoaderThread::RetrieveSample(const FMP4TrackSampleBufferPtr& InTrackSampleBuffer, const FTrackIterator& InSampleIt, const FTrackIterator& InOptionalTimecodeIt, const ElectraProtronUtils::FCodecInfo::FTMCDTimecode& InTimecodeInfo)
+TSharedPtr<FElectraProtronPlayer::FImpl::FMP4Sample, ESPMode::ThreadSafe> FElectraProtronPlayer::FImpl::FLoaderThread::RetrieveSample(const FMP4TrackSampleBufferPtr& InTrackSampleBuffer, const FTrackIterator& InSampleIt, const FTrackIterator& InOptionalTimecodeIt, const Electra::FCodecTypeFormat::FTMCDTimecode& InTimecodeInfo)
 {
 	InTrackSampleBuffer->Lock.Lock();
 	FMP4SamplePtr* CurrentSample = InTrackSampleBuffer->SampleMap.Find(InSampleIt->GetSampleNumber());
diff --git a/Engine/Plugins/Media/ElectraPlayer/Source/ElectraProtron/Private/Player/TrackFormatInfo.h b/Engine/Plugins/Media/ElectraPlayer/Source/ElectraProtron/Private/Player/TrackFormatInfo.h
deleted file mode 100644
diff --git a/Engine/Plugins/Media/ElectraUtil/Source/ElectraBase/Private/Utilities/UtilitiesMP4.cpp b/Engine/Plugins/Media/ElectraUtil/Source/ElectraBase/Private/Utilities/UtilitiesMP4.cpp
--- a/Engine/Plugins/Media/ElectraUtil/Source/ElectraBase/Private/Utilities/UtilitiesMP4.cpp
+++ b/Engine/Plugins/Media/ElectraUtil/Source/ElectraBase/Private/Utilities/UtilitiesMP4.cpp
@@ -555,4 +555,80 @@ FString FFileDataReader::GetLastError()
 	return LastError;
 }
 
+
+/**********************************************************************************************************************/
+/**********************************************************************************************************************/
+/**********************************************************************************************************************/
+
+class FBufferDataReader : public IBufferDataReader
+{
+public:
+	~FBufferDataReader() = default;
+	FBufferDataReader(TConstArrayView<const uint8>& InDataBuffer)
+		: DataBuffer(MoveTemp(InDataBuffer))
+	{
+		TotalFileSize = DataBuffer.Num();
+	}
+
+	int64 ReadData(void* InOutBuffer, int64 InNumBytes, int64 InFromOffset, FCancellationCheckDelegate InCheckCancellationDelegate) override;
+	int64 GetTotalFileSize() override;
+	int64 GetCurrentFileOffset() override;
+	bool HasReachedEOF() override;
+	FString GetLastError() override;
+private:
+	TConstArrayView<const uint8> DataBuffer;
+	int64 TotalFileSize = 0;
+	int64 CurrentOffset = 0;
+};
+
+TSharedPtr<IBufferDataReader, ESPMode::ThreadSafe> IBufferDataReader::Create(TConstArrayView<const uint8> InData)
+{
+	return MakeShared<FBufferDataReader, ESPMode::ThreadSafe>(InData);
+}
+
+int64 FBufferDataReader::ReadData(void* InOutBuffer, int64 InNumBytes, int64 InFromOffset, FCancellationCheckDelegate InCheckCancellationDelegate)
+{
+	check(InNumBytes >= 0);
+	if (InNumBytes <= 0)
+	{
+		return 0;
+	}
+	if (CurrentOffset != InFromOffset)
+	{
+		check(InFromOffset >= 0 && InFromOffset <= TotalFileSize);
+		CurrentOffset = InFromOffset >= 0 ? InFromOffset <= TotalFileSize ? InFromOffset : TotalFileSize : 0;
+	}
+	if (InNumBytes + CurrentOffset > TotalFileSize)
+	{
+		InNumBytes = TotalFileSize - CurrentOffset;
+	}
+	if (InOutBuffer)
+	{
+		FMemory::Memcpy(InOutBuffer, DataBuffer.GetData() + CurrentOffset, InNumBytes);
+	}
+	CurrentOffset += InNumBytes;
+	return InNumBytes;
+}
+
+int64 FBufferDataReader::GetTotalFileSize()
+{
+	return TotalFileSize;
+}
+
+int64 FBufferDataReader::GetCurrentFileOffset()
+{
+	return CurrentOffset;
+}
+
+bool FBufferDataReader::HasReachedEOF()
+{
+	return CurrentOffset >= TotalFileSize;
+}
+
+FString FBufferDataReader::GetLastError()
+{
+	return FString();
+}
+
+
 } // namespace Electra
diff --git a/Engine/Plugins/Media/ElectraUtil/Source/ElectraBase/Public/CodecTypeFormat.h b/Engine/Plugins/Media/ElectraUtil/Source/ElectraBase/Public/CodecTypeFormat.h
new file mode 100644
--- /dev/null
+++ b/Engine/Plugins/Media/ElectraUtil/Source/ElectraBase/Public/CodecTypeFormat.h
@@ -0,0 +1,135 @@
+// Copyright Epic Games, Inc. All Rights Reserved.
+
+#pragma once
+
+#include "CoreTypes.h"
+#include "Containers/Array.h"
+#include "Containers/UnrealString.h"
+#include "Misc/TVariant.h"
+#include "Misc/FrameRate.h"
+#include "Misc/Timecode.h"
+
+#include "Utilities/BCP47-Helpers.h"
+
+#define UE_API ELECTRABASE_API
+
+namespace Electra
+{
+
+	/**
+	 * This structure contains codec format related information.
+	 * Only information pertaining to the codec that can be derived from
+	 * the container, the RFC6381 format string or a media stream playlist
+	 * should be stored here.
+	 */
+	struct FCodecTypeFormat
+	{
+		enum class EType
+		{
+			Video,
+			Audio,
+			Subtitle,
+			Timecode,
+			Invalid,
+			MAX = Invalid
+		};
+		struct FVideo
+		{
+			// Codec specific profile, level, etc.
+			struct FProfile
+			{
+				uint32 Tier = 0;
+				uint32 ProfileSpace = 0;
+				uint32 Profile = 0;
+				uint32 Level = 0;
+				uint64 Constraints = 0;
+				uint64 CompatibilityFlags = 0;
+			};
+			// Codec specific color information. These do not necessarily correspond to MPEG definitions!
+			struct FColorInfo
+			{
+				uint8 colourPrimaries = 0;
+				uint8 transferCharacteristics = 0;
+				uint8 matrixCoefficients = 0;
+				uint8 videoFullRangeFlag = 0;
+				uint8 chromaSubsampling = 0;
+				struct FChromaInfo
+				{
+					uint8 monochrome = 0;
+					uint8 chromaSubsamplingX = 0;
+					uint8 chromaSubsamplingY = 0;
+					uint8 chromaSamplingPosition = 0;
+				};
+				TOptional<FChromaInfo> ChromaInfo;
+			};
+			FFrameRate FrameRate {0, 0};
+			uint32 Width = 0;
+			uint32 Height = 0;
+			uint32 BitDepth = 0;
+			FProfile Profile;
+			TOptional<FColorInfo> OptColorInfo;
+		};
+		struct FAudio
+		{
+			uint32 NumChannels = 0;
+			uint32 ChannelConfiguration = 0;
+			uint32 SampleRate = 0;
+		};
+		struct FSubtitle
+		{
+		};
+		struct FTMCDTimecode
+		{
+			// See: https://developer.apple.com/documentation/quicktime-file-format/timecode_sample_description/flags
+			enum EFlags : uint32
+			{
+				DropFrame = 0x0001,					// Indicates whether the timecode is drop frame. Set it to 1 if the timecode is drop frame.
+				Max24Hour = 0x0002,					// Indicates whether the timecode wraps after 24 hours. Set it to 1 if the timecode wraps.
+				AllowNegativeTimes = 0x0004,		// Indicates whether negative time values are allowed. Set it to 1 if the timecode supports negative values.
+				Counter = 0x0008					// Indicates whether the time value corresponds to a tape counter value. Set it to 1 if the timecode values are tape counter values.
+			};
+
+			uint32 Flags = 0;
+			uint32 Timescale = 0;
+			uint32 FrameDuration = 0;
+			uint32 NumberOfFrames = 0;
+
+			bool IsDropFrame() const
+			{ return !!(Flags & EFlags::DropFrame); }
+
+			bool WrapsAfter24Hours() const
+			{ return !!(Flags & EFlags::Max24Hour); }
+
+			bool SupportsNegativeTime() const
+			{ return !!(Flags & EFlags::AllowNegativeTimes); }
+
+			FFrameRate GetFrameRate() const
+			{ return IsDropFrame() ? FFrameRate(Timescale, FrameDuration) : FFrameRate(NumberOfFrames, 1); }
+
+			FTimecode ConvertToTimecode(uint32 InSampleTimecode) const
+			{
+				// Needs to be an int32 for use with the following methods.
+				check(InSampleTimecode <= 0x7fffffffU);
+				const FFrameRate FrameRate = GetFrameRate();
+				// Convert to time code (apply roll over, etc). via conversion to seconds first.
+				return FTimecode(FrameRate.AsSeconds(FFrameTime(FFrameNumber(static_cast<int32>(InSampleTimecode)))), FrameRate, IsDropFrame(), WrapsAfter24Hours());
+			}
+
+		};
+		TVariant<FVideo,FAudio,FSubtitle,FTMCDTimecode> Properties;
+		BCP47::FLanguageTag LanguageTag;
+		FString HumanReadableFormatInfo;
+		FString RFC6381;
+		FString MimeType;
+		TMap<uint32, TArray<uint8>> ExtraBoxes;
+		TArray<uint8> DCR;
+		TArray<uint8> CSD;
+		EType Type = EType::Invalid;
+		uint32 FourCC = 0;
+		uint32 Bitrate = 0;
+		uint32 AverageBitrate = 0;
+	};
+
+} // namespace Electra
+
+#undef UE_API
diff --git a/Engine/Plugins/Media/ElectraUtil/Source/ElectraBase/Public/Utilities/MP4Boxes/MP4BoxBase.h b/Engine/Plugins/Media/ElectraUtil/Source/ElectraBase/Public/Utilities/MP4Boxes/MP4BoxBase.h
--- a/Engine/Plugins/Media/ElectraUtil/Source/ElectraBase/Public/Utilities/MP4Boxes/MP4BoxBase.h
+++ b/Engine/Plugins/Media/ElectraUtil/Source/ElectraBase/Public/Utilities/MP4Boxes/MP4BoxBase.h
@@ -170,7 +170,7 @@ namespace Electra
 				return false;
 			}
 
-			virtual void ProcessBoxChildrenRecursively(FMP4AtomReader& InOutReader, const FMP4BoxInfo& InCurrentBoxInfo);
+			ELECTRABASE_API virtual void ProcessBoxChildrenRecursively(FMP4AtomReader& InOutReader, const FMP4BoxInfo& InCurrentBoxInfo);
 
 			void SetRootBoxData(const TSharedPtr<FMP4BoxInfo, ESPMode::ThreadSafe>& InRootBoxData)
 			{ RootBoxData = InRootBoxData; }
diff --git a/Engine/Plugins/Media/ElectraUtil/Source/ElectraBase/Public/Utilities/UtilitiesMP4.h b/Engine/Plugins/Media/ElectraUtil/Source/ElectraBase/Public/Utilities/UtilitiesMP4.h
--- a/Engine/Plugins/Media/ElectraUtil/Source/ElectraBase/Public/Utilities/UtilitiesMP4.h
+++ b/Engine/Plugins/Media/ElectraUtil/Source/ElectraBase/Public/Utilities/UtilitiesMP4.h
@@ -40,6 +40,23 @@ namespace Electra
 	static constexpr uint64 GetFromBigEndian(uint64 value)		{ return EndianSwap(value); }
 #endif
 
+
+	class IFileDataReader : public IBaseDataReader
+	{
+	public:
+		static ELECTRABASE_API TSharedPtr<IFileDataReader, ESPMode::ThreadSafe> Create();
+		virtual ~IFileDataReader() = default;
+		virtual bool Open(const FString& InFilename) = 0;
+	};
+
+	class IBufferDataReader : public IBaseDataReader
+	{
+	public:
+		static ELECTRABASE_API TSharedPtr<IBufferDataReader, ESPMode::ThreadSafe> Create(TConstArrayView<const uint8> InData);
+		virtual ~IBufferDataReader() = default;
+	};
+
+
 	namespace UtilitiesMP4
 	{
 		class FMP4BoxBase;
@@ -160,6 +177,11 @@ namespace Electra
 			UE_API bool LocateAndReadRootBoxes(TArray<TSharedPtr<FMP4BoxData, ESPMode::ThreadSafe>>& OutBoxInfos, const TSharedPtr<IBaseDataReader, ESPMode::ThreadSafe>& InDataReader, const TArray<uint32>& InFirstBoxes, const TArray<uint32>& InStopAfterBoxes, const TArray<uint32>& InReadDataOfBoxes, IBaseDataReader::FCancellationCheckDelegate InCheckCancellationDelegate);
 			FString GetLastError() const
 			{ return LastError; }
+
+			static bool LocateAndReadRootBoxesFromBuffer(TArray<TSharedPtr<FMP4BoxData, ESPMode::ThreadSafe>>& OutBoxInfos, TConstArrayView<const uint8> InData)
+			{
+				return FMP4BoxLocatorReader().LocateAndReadRootBoxes(OutBoxInfos, IBufferDataReader::Create(InData), TArray<uint32>(), TArray<uint32>(), TArray<uint32>(), IBaseDataReader::FCancellationCheckDelegate());
+			}
 		private:
 			FString LastError;
 			int64 CurrentOffset = 0;
@@ -179,16 +201,6 @@ namespace Electra
 
 	} // namespace UtilitiesMP4
 
-
-
-	class IFileDataReader : public IBaseDataReader
-	{
-	public:
-		static ELECTRABASE_API TSharedPtr<IFileDataReader, ESPMode::ThreadSafe> Create();
-		virtual ~IFileDataReader() = default;
-		virtual bool Open(const FString& InFilename) = 0;
-	};
-
 } // namespace Electra
 
 #undef UE_API
