Fix a potential GPU hang due to out of order D3D12 fence signal and waits across queues in Electra video decoding

#Synced-CL 49076157

diff --git a/Engine/Plugins/Media/ElectraCodecs/Source/ElectraDecoders/Private/Windows/WindowsElectraDecodersPlatformResources.cpp b/Engine/Plugins/Media/ElectraCodecs/Source/ElectraDecoders/Private/Windows/WindowsElectraDecodersPlatformResources.cpp
--- a/Engine/Plugins/Media/ElectraCodecs/Source/ElectraDecoders/Private/Windows/WindowsElectraDecodersPlatformResources.cpp
+++ b/Engine/Plugins/Media/ElectraCodecs/Source/ElectraDecoders/Private/Windows/WindowsElectraDecodersPlatformResources.cpp
@@ -116,12 +116,12 @@ bool FElectraDecodersPlatformResourcesWindows::RunCodeAsync(TFunction<void()>&&
 
 
 
-static void TriggerDataCopy(TRefCountPtr<ID3D12Resource> InDstResource, TRefCountPtr<ID3D12Resource> InSrcResource, TRefCountPtr<ID3D12GraphicsCommandList> D3DCmdList, TRefCountPtr<ID3D12CommandAllocator> D3DCmdAllocator, TRefCountPtr<ID3D12Fence> D3DFence, uint64 FenceValue, const FElectraDecoderOutputSync& OutputSync)
+static void TriggerDataCopy(TRefCountPtr<ID3D12Resource> InDstResource, TRefCountPtr<ID3D12Resource> InSrcResource, TRefCountPtr<ID3D12GraphicsCommandList> D3DCmdList, TRefCountPtr<ID3D12CommandAllocator> D3DCmdAllocator, TRefCountPtr<ID3D12Fence> D3DFence, uint64 FenceValue, const TSharedPtr<std::atomic<bool>>& FenceSignalled, const FElectraDecoderOutputSync& OutputSync)
 {
 	// Trigger copy (this will eventually execute on the submission thread of RHI if running in UE)
 	// (note: we pass in all of FElectraDecoderOutputSync to guarantee any references needed to make the decoder output sync work are passed along, too!)
 	GetID3D12PlatformDynamicRHI()->RHIRunOnQueue(ED3D12RHIRunOnQueueType::Copy,
-		[DstResource = InDstResource, SrcResource = InSrcResource, CmdList = D3DCmdList, CmdAllocator = D3DCmdAllocator, DestFence = D3DFence, DestFenceValue = FenceValue, OutputSync](ID3D12CommandQueue* D3DCmdQueue)
+		[DstResource = InDstResource, SrcResource = InSrcResource, CmdList = D3DCmdList, CmdAllocator = D3DCmdAllocator, DestFence = D3DFence, DestFenceValue = FenceValue, DestFenceSignalled = FenceSignalled, OutputSync](ID3D12CommandQueue* D3DCmdQueue)
 		{
 			TRefCountPtr<ID3D12Fence> ResourceFence;
 	#if ELECTRA_DECODERS_ENABLE_DX && ALLOW_MFSAMPLE_WITH_DX12
@@ -170,6 +170,8 @@ static void TriggerDataCopy(TRefCountPtr<ID3D12Resource> InDstResource, TRefCoun
 
 			// Notify user of the output data we just copied of its arrival
 			D3DCmdQueue->Signal(DestFence, DestFenceValue);
+			*DestFenceSignalled = true;
+
 		}, false);
 }
 
@@ -357,6 +359,7 @@ static bool SetupOutputTextureSampleFromDX12Resource(FString& OutErrorMessage, T
 		#endif
 		InOutTextureSampleToSetup->D3DFence = OutputData.Fence;
 		InOutTextureSampleToSetup->FenceValue = OutputData.FenceValue;
+		InOutTextureSampleToSetup->FenceSignalled = MakeShared<std::atomic<bool>>(false);
 
 		InOutTextureSampleToSetup->TextureDX12Dim = TextureDim;
 	}
@@ -451,16 +454,16 @@ static bool SetupOutputTextureSampleFromDX12Resource(FString& OutErrorMessage, T
 	if (OutputSync.TaskSync.IsValid())
 	{
 		// Note: we capture "this" as we ensure that this instance only dies once the copy triggered here is actually done, hence ensuring any reference to "this" is done
-		bTriggerOk = FElectraDecodersPlatformResources::RunCodeAsync([DstResource=InOutTextureSampleToSetup->TextureDX12, SrcResource=InResource, D3DCmdList=InOutTextureSampleToSetup->D3DCmdList, D3DCmdAllocator=InOutTextureSampleToSetup->D3DCmdAllocator, D3DFence=InOutTextureSampleToSetup->D3DFence, FenceValue=InOutTextureSampleToSetup->FenceValue, OutputSync]()
+		bTriggerOk = FElectraDecodersPlatformResources::RunCodeAsync([DstResource=InOutTextureSampleToSetup->TextureDX12, SrcResource=InResource, D3DCmdList=InOutTextureSampleToSetup->D3DCmdList, D3DCmdAllocator=InOutTextureSampleToSetup->D3DCmdAllocator, D3DFence=InOutTextureSampleToSetup->D3DFence, FenceValue=InOutTextureSampleToSetup->FenceValue, FenceSignalled=InOutTextureSampleToSetup->FenceSignalled, OutputSync]()
 			{
-				TriggerDataCopy(DstResource, SrcResource, D3DCmdList, D3DCmdAllocator, D3DFence, FenceValue, OutputSync);
+				TriggerDataCopy(DstResource, SrcResource, D3DCmdList, D3DCmdAllocator, D3DFence, FenceValue, FenceSignalled, OutputSync);
 			}, OutputSync.TaskSync);
 	}
 
 	if (!bTriggerOk)
 	{
 		// We could not run the trigger async. Schedule the copy right away. Any needed synchronization will be done in the copy-queue by the GPU
-		TriggerDataCopy(InOutTextureSampleToSetup->TextureDX12, InResource, InOutTextureSampleToSetup->D3DCmdList, InOutTextureSampleToSetup->D3DCmdAllocator, InOutTextureSampleToSetup->D3DFence, InOutTextureSampleToSetup->FenceValue, OutputSync);
+		TriggerDataCopy(InOutTextureSampleToSetup->TextureDX12, InResource, InOutTextureSampleToSetup->D3DCmdList, InOutTextureSampleToSetup->D3DCmdAllocator, InOutTextureSampleToSetup->D3DFence, InOutTextureSampleToSetup->FenceValue, InOutTextureSampleToSetup->FenceSignalled, OutputSync);
 	}
 
 	// Finish initialization.
diff --git a/Engine/Plugins/Media/ElectraUtil/Source/ElectraSamples/Private/Windows/ElectraTextureSample_Windows.cpp b/Engine/Plugins/Media/ElectraUtil/Source/ElectraSamples/Private/Windows/ElectraTextureSample_Windows.cpp
--- a/Engine/Plugins/Media/ElectraUtil/Source/ElectraSamples/Private/Windows/ElectraTextureSample_Windows.cpp
+++ b/Engine/Plugins/Media/ElectraUtil/Source/ElectraSamples/Private/Windows/ElectraTextureSample_Windows.cpp
@@ -547,7 +547,8 @@ bool FElectraTextureSample::Convert(FRHICommandListImmediate& RHICmdList, FTextu
 	if (bHaveTexture && (RHIGetInterfaceType() == ERHIInterfaceType::D3D12))
 	{
 		uint64 SyncFenceValue;
-		TRefCountPtr<IUnknown> SyncCommon = GetSync(SyncFenceValue);
+		TSharedPtr<std::atomic<bool>> bFenceSignalled;
+		TRefCountPtr<IUnknown> SyncCommon = GetSync(SyncFenceValue, bFenceSignalled);
 		if (SyncCommon)
 		{
 			//
@@ -557,9 +558,9 @@ bool FElectraTextureSample::Convert(FRHICommandListImmediate& RHICmdList, FTextu
 			HRESULT Result = SyncCommon->QueryInterface(__uuidof(ID3D12Fence), (void**)SyncFence.GetInitReference());
 			if (SUCCEEDED(Result))
 			{
-				RHICmdList.EnqueueLambda([SyncFence, SyncFenceValue](FRHICommandList& RHICmdList)
+				RHICmdList.EnqueueLambda([SyncFence, SyncFenceValue, bFenceSignalled](FRHICommandList& RHICmdList)
 				{
-					GetID3D12DynamicRHI()->RHIWaitManualFence(RHICmdList, SyncFence, SyncFenceValue);
+					GetID3D12DynamicRHI()->RHIWaitManualFence(RHICmdList, SyncFence, SyncFenceValue, bFenceSignalled);
 				});
 				return true;
 			}
@@ -640,8 +641,9 @@ bool FElectraTextureSample::Convert(FRHICommandListImmediate& RHICmdList, FTextu
 	}
 
 	uint64 SyncValue = 0;
+	TSharedPtr<std::atomic<bool>> bFenceSignalled;
 	TRefCountPtr<IUnknown> TextureCommon = GetTextureD3D();
-	TRefCountPtr<IUnknown> SyncCommon = GetSync(SyncValue);
+	TRefCountPtr<IUnknown> SyncCommon = GetSync(SyncValue, bFenceSignalled);
 
 #ifdef ELECTRA_HAVE_DX11
 	// DX11 texture?
@@ -696,18 +698,20 @@ float FElectraTextureSample::GetSampleDataScale(bool b10Bit) const
 }
 
 
-TRefCountPtr<IUnknown> FElectraTextureSample::GetSync(uint64& OutSyncValue)
+TRefCountPtr<IUnknown> FElectraTextureSample::GetSync(uint64& OutSyncValue, TSharedPtr<std::atomic<bool>>& OutFenceSignalled)
 {
 	if (TextureDX11)
 	{
 		TRefCountPtr<IUnknown> KeyedMutex;
 		/*HRESULT Result =*/ TextureDX11->QueryInterface(_uuidof(IDXGIKeyedMutex), (void**)&KeyedMutex);
 		OutSyncValue = 0;
+		OutFenceSignalled = nullptr;
 		return KeyedMutex;
 	}
 	else if (TextureDX12)
 	{
 		OutSyncValue = FenceValue;
+		OutFenceSignalled = FenceSignalled;
 		return static_cast<IUnknown*>(D3DFence.GetReference());
 	}
 	return nullptr;
diff --git a/Engine/Plugins/Media/ElectraUtil/Source/ElectraSamples/Public/Windows/ElectraTextureSample.h b/Engine/Plugins/Media/ElectraUtil/Source/ElectraSamples/Public/Windows/ElectraTextureSample.h
--- a/Engine/Plugins/Media/ElectraUtil/Source/ElectraSamples/Public/Windows/ElectraTextureSample.h
+++ b/Engine/Plugins/Media/ElectraUtil/Source/ElectraSamples/Public/Windows/ElectraTextureSample.h
@@ -70,7 +70,7 @@ public:
 		return ConverterInfoFlags_PreprocessOnly;
 	}
 
-	ELECTRASAMPLES_API TRefCountPtr<IUnknown> GetSync(uint64& OutSyncValue);
+	ELECTRASAMPLES_API TRefCountPtr<IUnknown> GetSync(uint64& OutSyncValue, TSharedPtr<std::atomic<bool>>& OutFenceSignalled);
 	ELECTRASAMPLES_API TRefCountPtr<IUnknown> GetTextureD3D() const;
 
 
@@ -102,6 +102,7 @@ public:
 	FIntPoint TextureDX12Dim {0, 0};
 	TRefCountPtr<ID3D12Fence> D3DFence;
 	uint64 FenceValue = 0;
+	TSharedPtr<std::atomic<bool>> FenceSignalled;
 	TRefCountPtr<ID3D12Resource> DecoderOutputResource;
 	TRefCountPtr<ID3D12CommandAllocator> D3DCmdAllocator;
 	TRefCountPtr<ID3D12GraphicsCommandList> D3DCmdList;
diff --git a/Engine/Source/Runtime/D3D12RHI/Private/D3D12CommandContext.cpp b/Engine/Source/Runtime/D3D12RHI/Private/D3D12CommandContext.cpp
--- a/Engine/Source/Runtime/D3D12RHI/Private/D3D12CommandContext.cpp
+++ b/Engine/Source/Runtime/D3D12RHI/Private/D3D12CommandContext.cpp
@@ -255,17 +255,17 @@ void FD3D12ContextCommon::SignalManualFence(ID3D12Fence* Fence, uint64 Value)
 		CloseCommandList();
 	}
 
-	GetPayload(EPhase::Signal)->ManualFencesToSignal.Emplace(Fence, Value);
+	GetPayload(EPhase::Signal)->ManualFencesToSignal.Emplace(Fence, Value, nullptr);
 }
 
-void FD3D12ContextCommon::WaitManualFence(ID3D12Fence* Fence, uint64 Value)
+void FD3D12ContextCommon::WaitManualFence(ID3D12Fence* Fence, uint64 Value, const TSharedPtr<std::atomic<bool>>& bHasSignalled)
 {
 	if (IsPendingCommands())
 	{
 		CloseCommandList();
 	}
 
-	GetPayload(EPhase::Wait)->ManualFencesToWait.Emplace(Fence, Value);
+	GetPayload(EPhase::Wait)->ManualFencesToWait.Emplace(Fence, Value, bHasSignalled);
 }
 
 FD3D12QueryLocation FD3D12ContextCommon::AllocateQuery(ED3D12QueryType Type, void* Target)
diff --git a/Engine/Source/Runtime/D3D12RHI/Private/D3D12CommandContext.h b/Engine/Source/Runtime/D3D12RHI/Private/D3D12CommandContext.h
--- a/Engine/Source/Runtime/D3D12RHI/Private/D3D12CommandContext.h
+++ b/Engine/Source/Runtime/D3D12RHI/Private/D3D12CommandContext.h
@@ -230,7 +230,7 @@ public:
 	void SignalManualFence(ID3D12Fence* Fence, uint64 Value);
 
 	// Inserts a command that waits the specified D3D12 fence object.
-	void WaitManualFence(ID3D12Fence* Fence, uint64 Value);
+	void WaitManualFence(ID3D12Fence* Fence, uint64 Value, const TSharedPtr<std::atomic<bool>>& bHasSignalled);
 
 	// Inserts a timestamp query command. "Target" specifies the optional 
 	// location the result will be written to by the interrupt handler thread.
diff --git a/Engine/Source/Runtime/D3D12RHI/Private/D3D12DirectCommandListManager.cpp b/Engine/Source/Runtime/D3D12RHI/Private/D3D12DirectCommandListManager.cpp
--- a/Engine/Source/Runtime/D3D12RHI/Private/D3D12DirectCommandListManager.cpp
+++ b/Engine/Source/Runtime/D3D12RHI/Private/D3D12DirectCommandListManager.cpp
@@ -192,7 +192,7 @@ void FD3D12ManualFence::AdvanceBOP()
 	for (auto& [Queue, Fence] : Fences)
 	{
 		FD3D12Payload* Payload = new FD3D12Payload(*Queue);
-		Payload->ManualFencesToSignal.Emplace(Fence.GetReference(), NextValue);
+		Payload->ManualFencesToSignal.Emplace(Fence.GetReference(), NextValue, nullptr);
 
 		Payloads.Add(Payload);
 	}
diff --git a/Engine/Source/Runtime/D3D12RHI/Private/D3D12RHI.cpp b/Engine/Source/Runtime/D3D12RHI/Private/D3D12RHI.cpp
--- a/Engine/Source/Runtime/D3D12RHI/Private/D3D12RHI.cpp
+++ b/Engine/Source/Runtime/D3D12RHI/Private/D3D12RHI.cpp
@@ -753,13 +753,13 @@ void FD3D12DynamicRHI::RHISignalManualFence(FRHICommandList& RHICmdList, ID3D12F
 	Context.SignalManualFence(Fence, Value);
 }
 
-void FD3D12DynamicRHI::RHIWaitManualFence(FRHICommandList& RHICmdList, ID3D12Fence* Fence, uint64 Value)
+void FD3D12DynamicRHI::RHIWaitManualFence(FRHICommandList& RHICmdList, ID3D12Fence* Fence, uint64 Value, const TSharedPtr<std::atomic<bool>>& bHasSignalled)
 {
 	checkf(FRHIGPUMask::All() == FRHIGPUMask::GPU0(), TEXT("RHIWaitManualFence cannot be used by multi-GPU code"));
 	const uint32 GPUIndex = 0;
 
 	FD3D12CommandContext& Context = FD3D12CommandContext::Get(RHICmdList, GPUIndex);
-	Context.WaitManualFence(Fence, Value);
+	Context.WaitManualFence(Fence, Value, bHasSignalled);
 }
 
 void FD3D12DynamicRHI::RHIVerifyResult(ID3D12Device* Device, HRESULT Result, const ANSICHAR* Code, const ANSICHAR* Filename, uint32 Line, FString Message) const
diff --git a/Engine/Source/Runtime/D3D12RHI/Private/D3D12RHIPrivate.h b/Engine/Source/Runtime/D3D12RHI/Private/D3D12RHIPrivate.h
--- a/Engine/Source/Runtime/D3D12RHI/Private/D3D12RHIPrivate.h
+++ b/Engine/Source/Runtime/D3D12RHI/Private/D3D12RHIPrivate.h
@@ -473,7 +473,7 @@ public:
 	virtual void RHIFinishExternalComputeWork(FRHICommandListBase& ExecutingCmdList, uint32 InDeviceIndex, ID3D12GraphicsCommandList* InCommandList) final override;
 	virtual void RHITransitionResource(FRHICommandList& RHICmdList, FRHITexture* InTexture, D3D12_RESOURCE_STATES InState, uint32 InSubResource) final override;
 	virtual void RHISignalManualFence(FRHICommandList& RHICmdList, ID3D12Fence* Fence, uint64 Value) final override;
-	virtual void RHIWaitManualFence(FRHICommandList& RHICmdList, ID3D12Fence* Fence, uint64 Value) final override;
+	virtual void RHIWaitManualFence(FRHICommandList& RHICmdList, ID3D12Fence* Fence, uint64 Value, const TSharedPtr<std::atomic<bool>>& bHasSignalled) final override;
 	virtual void RHIVerifyResult(ID3D12Device* Device, HRESULT Result, const ANSICHAR* Code, const ANSICHAR* Filename, uint32 Line, FString Message = FString()) const final override;
 
 #if D3D12_RHI_RAYTRACING
diff --git a/Engine/Source/Runtime/D3D12RHI/Private/D3D12Submission.cpp b/Engine/Source/Runtime/D3D12RHI/Private/D3D12Submission.cpp
--- a/Engine/Source/Runtime/D3D12RHI/Private/D3D12Submission.cpp
+++ b/Engine/Source/Runtime/D3D12RHI/Private/D3D12Submission.cpp
@@ -413,6 +413,23 @@ FD3D12DynamicRHI::FProcessResult FD3D12DynamicRHI::ProcessSubmissionQueue()
 						Payload->SyncPointsToWait.Index++;
 						bProgress = true;
 					}
+
+					while (Payload->ManualFencesToWait.Index < Payload->ManualFencesToWait.Num())
+					{
+						FD3D12PayloadBase::FManualFence& ManualFence = Payload->ManualFencesToWait[Payload->ManualFencesToWait.Index];
+						if (ManualFence.bHasSignalled != nullptr && !(*ManualFence.bHasSignalled))
+						{
+							// The CPU-side submit action for a manual fence has not happened yet.
+							// Don't process this payload until that has, otherwise we can deadlock due to inversion of submission order.
+							// Skip processing this queue, and move on to the next. We will retry later when
+							// further work is submitted, which may contain the fence we need.
+							Result.Status |= EQueueStatus::Pending;
+							return;
+						}
+
+						Payload->ManualFencesToWait.Index++;
+						bProgress = true;
+					}
 				
 					// All necessary sync points have been resolved.
 					Payload->SyncPointsToWait = {};
@@ -670,8 +687,9 @@ void FD3D12DynamicRHI::FlushBatchedPayloads(FD3D12Queue::FPayloadArray& Payloads
 		}
 
 		// Wait for manual fences
-		for (auto& [LocalFence, Value] : Payload->ManualFencesToWait)
+		for (auto& [LocalFence, Value, bFenceSignalled] : Payload->ManualFencesToWait)
 		{
+			checkf(!bFenceSignalled.IsValid() || *bFenceSignalled, TEXT("The signal operation for a provided manual D3D12 fence has not yet happened."));
 			VERIFYD3D12RESULT(Queue.D3DCommandQueue->Wait(LocalFence, Value));
 		}
 	};
@@ -762,9 +780,13 @@ void FD3D12DynamicRHI::FlushBatchedPayloads(FD3D12Queue::FPayloadArray& Payloads
 		FD3D12Queue& Queue = Payload->Queue;
 
 		// Signal any manual fences
-		for (auto& [ManualFence, Value] : Payload->ManualFencesToSignal)
+		for (auto& [ManualFence, Value, bFenceSignalled] : Payload->ManualFencesToSignal)
 		{
 			VERIFYD3D12RESULT(Queue.D3DCommandQueue->Signal(ManualFence, Value));
+			if (bFenceSignalled)
+			{
+				*bFenceSignalled = true;
+			}
 		}
 
 		// Signal the queue fence
diff --git a/Engine/Source/Runtime/D3D12RHI/Private/D3D12Submission.h b/Engine/Source/Runtime/D3D12RHI/Private/D3D12Submission.h
--- a/Engine/Source/Runtime/D3D12RHI/Private/D3D12Submission.h
+++ b/Engine/Source/Runtime/D3D12RHI/Private/D3D12Submission.h
@@ -218,10 +218,14 @@ struct FD3D12PayloadBase : public FD3D12PayloadBaseFixLayout
 		// The value to signal the fence with.
 		uint64 Value;
 
+		// Optional bool pointer to check whether this fence has been signalled
+		TSharedPtr<std::atomic<bool>> bHasSignalled;
+
 		FManualFence() = default;
-		FManualFence(TRefCountPtr<ID3D12Fence>&& Fence, uint64 Value)
+		FManualFence(TRefCountPtr<ID3D12Fence>&& Fence, uint64 Value, const TSharedPtr<std::atomic<bool>>& bHasSignalled)
 			: Fence(MoveTemp(Fence))
 			, Value(Value)
+			, bHasSignalled(bHasSignalled)
 		{}
 	};
 
@@ -243,7 +247,11 @@ struct FD3D12PayloadBase : public FD3D12PayloadBaseFixLayout
 		uint64 Value;
 	};
 	TArray<FQueueFence, TInlineAllocator<GD3D12MaxNumQueues>> QueueFencesToWait;
-	TArray<FManualFence> ManualFencesToWait;
+	struct : TArray<FManualFence>
+	{
+		int32 Index = 0;
+
+	} ManualFencesToWait;
 
 	void AddQueueFenceWait(FD3D12Fence& Fence, uint64 Value);
 
diff --git a/Engine/Source/Runtime/D3D12RHI/Public/ID3D12DynamicRHI.h b/Engine/Source/Runtime/D3D12RHI/Public/ID3D12DynamicRHI.h
--- a/Engine/Source/Runtime/D3D12RHI/Public/ID3D12DynamicRHI.h
+++ b/Engine/Source/Runtime/D3D12RHI/Public/ID3D12DynamicRHI.h
@@ -55,7 +55,7 @@ struct ID3D12DynamicRHI : public FDynamicRHI
 	virtual void                  RHITransitionResource(FRHICommandList& RHICmdList, FRHITexture* InTexture, D3D12_RESOURCE_STATES InState, uint32 InSubResource) = 0;
 
 	virtual void                  RHISignalManualFence(FRHICommandList& RHICmdList, ID3D12Fence* Fence, uint64 Value) = 0;
-	virtual void                  RHIWaitManualFence(FRHICommandList& RHICmdList, ID3D12Fence* Fence, uint64 Value) = 0;
+	virtual void                  RHIWaitManualFence(FRHICommandList& RHICmdList, ID3D12Fence* Fence, uint64 Value, const TSharedPtr<std::atomic<bool>>& bHasSignalled) = 0;
 
 	virtual void                  RHIVerifyResult(ID3D12Device* Device, HRESULT Result, const ANSICHAR* Code, const ANSICHAR* Filename, uint32 Line, FString Message = FString()) const = 0;
 
