Fix stale fence signalled atomic in Electra D3D12 decoding

#Synced-CL 49104626

diff --git a/Engine/Plugins/Media/ElectraCodecs/Source/ElectraDecoders/Private/Windows/WindowsElectraDecodersPlatformResources.cpp b/Engine/Plugins/Media/ElectraCodecs/Source/ElectraDecoders/Private/Windows/WindowsElectraDecodersPlatformResources.cpp
--- a/Engine/Plugins/Media/ElectraCodecs/Source/ElectraDecoders/Private/Windows/WindowsElectraDecodersPlatformResources.cpp
+++ b/Engine/Plugins/Media/ElectraCodecs/Source/ElectraDecoders/Private/Windows/WindowsElectraDecodersPlatformResources.cpp
@@ -170,6 +170,8 @@ static void TriggerDataCopy(TRefCountPtr<ID3D12Resource> InDstResource, TRefCoun
 
 			// Notify user of the output data we just copied of its arrival
 			D3DCmdQueue->Signal(DestFence, DestFenceValue);
+			checkf(!*DestFenceSignalled, TEXT("Electra texture sample copy fence should not have been signalled yet at this point"));
+
 			*DestFenceSignalled = true;
 
 		}, false);
@@ -359,7 +361,6 @@ static bool SetupOutputTextureSampleFromDX12Resource(FString& OutErrorMessage, T
 		#endif
 		InOutTextureSampleToSetup->D3DFence = OutputData.Fence;
 		InOutTextureSampleToSetup->FenceValue = OutputData.FenceValue;
-		InOutTextureSampleToSetup->FenceSignalled = MakeShared<std::atomic<bool>>(false);
 
 		InOutTextureSampleToSetup->TextureDX12Dim = TextureDim;
 	}
@@ -369,6 +370,8 @@ static bool SetupOutputTextureSampleFromDX12Resource(FString& OutErrorMessage, T
 		InOutTextureSampleToSetup->D3DFence = InOutTextureSampleToSetup->D3D12ResourcePool->GetUpdatedBufferFence(InOutTextureSampleToSetup->FenceValue);
 	}
 
+	InOutTextureSampleToSetup->FenceSignalled = MakeShared<std::atomic<bool>>(false);
+
 	// - - - - - - - - - - - - - - - - - - - - - - - - - - - - - - - - - - - - - - - - - - - - - - - - - - - - - - - -
 	// Build command list to copy data from decoder buffer
 
