ElectraCodecs: Added a new feature name

#rnx

#Synced-CL 49644370

diff --git a/Engine/Plugins/Media/ElectraCodecs/Source/ElectraDecoders/Public/IElectraDecoderFeaturesAndOptions.h b/Engine/Plugins/Media/ElectraCodecs/Source/ElectraDecoders/Public/IElectraDecoderFeaturesAndOptions.h
--- a/Engine/Plugins/Media/ElectraCodecs/Source/ElectraDecoders/Public/IElectraDecoderFeaturesAndOptions.h
+++ b/Engine/Plugins/Media/ElectraCodecs/Source/ElectraDecoders/Public/IElectraDecoderFeaturesAndOptions.h
@@ -64,6 +64,16 @@ namespace IElectraDecoderFeature
 	 */
 	const TCHAR* const StartcodeToLength = TEXT("startcode_to_length");
 
+	/**
+	 * FVariant(bool)
+	 *    true - input data gets modified in place
+	 *    false - input data remains unchanged
+	 *
+	 * If the input access unit is modified in place this should be set to true
+	 * to have the higher up decoder layer create a copy of the data if necessary.
+	 */
+	const TCHAR* const ModifiesInputAUInPlace = TEXT("modifies_input_au_in_place");
+
 	/**
 	 * FVariant(FTimespan)
 	 *    Time offset with which to feed decoder input earlier
