ElectraCodecs: Added APV bitstream parser and decoder configuration record

#rnx

#Synced-CL 49644370

diff --git a/Engine/Plugins/Media/APVDecoderElectra/Source/APVDecoderElectra/Private/APVDecoder/ElectraMediaAPVDecoder.cpp b/Engine/Plugins/Media/APVDecoderElectra/Source/APVDecoderElectra/Private/APVDecoder/ElectraMediaAPVDecoder.cpp
--- a/Engine/Plugins/Media/APVDecoderElectra/Source/APVDecoderElectra/Private/APVDecoder/ElectraMediaAPVDecoder.cpp
+++ b/Engine/Plugins/Media/APVDecoderElectra/Source/APVDecoderElectra/Private/APVDecoder/ElectraMediaAPVDecoder.cpp
@@ -13,7 +13,9 @@
 #include "IElectraDecoder.h"
 #include "IElectraDecoderOutputVideo.h"
 #include "ElectraDecodersUtils.h"
-#include "Utils/ElectraBitstreamProcessorDefault.h"
+#include "Utils/OpenAPV/ElectraBitstreamProcessor_APV.h"
+#include "Utils/OpenAPV/ElectraUtilsAPVVideo.h"
+#include "Utilities/UtilitiesMP4.h"
 
 #include "ElectraDecodersPlatformResources.h"
 
@@ -252,7 +254,12 @@ public:
 	TSharedPtr<IElectraDecoderOutput, ESPMode::ThreadSafe> GetOutput() override;
 
 	TSharedPtr<IElectraDecoderBitstreamProcessor, ESPMode::ThreadSafe> CreateBitstreamProcessor() override
-	{ return FElectraDecoderBitstreamProcessorDefault::Create(); }
+	{
+		TMap<FString, FVariant> DecoderFeatures;
+		GetFeatures(DecoderFeatures);
+		DecoderFeatures.Emplace(IElectraDecoderFeature::ModifiesInputAUInPlace, FVariant(true));
+		return FElectraDecoderBitstreamProcessorAPV::Create(DecoderFeatures, InitialCreationOptions);
+	}
 
 	void Suspend() override
 	{ }
@@ -265,6 +272,8 @@ private:
 	int32 CreateDecoder();
 	void DestroyDecoder();
 
+	TMap<FString, FVariant> InitialCreationOptions;
+
 	int32 AspectW = 0;
 	int32 AspectH = 0;
 	uint32 Codec4CC = 0;
@@ -380,6 +389,7 @@ void FElectraMediaAPVDecoder::Shutdown()
 /*********************************************************************************************************************/
 
 FVideoDecoderAPVElectra::FVideoDecoderAPVElectra(const TMap<FString, FVariant>& InOptions)
+	: InitialCreationOptions(InOptions)
 {
 	AspectW = (int32)ElectraDecodersUtil::GetVariantValueSafeI64(InOptions, TEXT("aspect_w"), 0);
 	AspectH = (int32)ElectraDecodersUtil::GetVariantValueSafeI64(InOptions, TEXT("aspect_h"), 0);
@@ -467,38 +477,26 @@ IElectraDecoder::EDecoderError FVideoDecoderAPVElectra::DecodeAccessUnit(const F
 	// Decode data. This immediately produces a new output frame.
 	if (InInputAccessUnit.Data && InInputAccessUnit.DataSize > 4)
 	{
-		oapv_au_info_t AccessUnitInfo;
-		FMemory::Memset(&AccessUnitInfo, 0, sizeof(oapv_au_info_t));
-
-		// FIXME: Skip over the length value for now since we only handle a single AU here
-		void* pAU = (void*)((const uint8*)InInputAccessUnit.Data + 4);
-		int nAU = static_cast<int>(InInputAccessUnit.DataSize - 4);
-
-		int Result = oapvd_info(/*const_cast<void*>*/pAU, nAU, &AccessUnitInfo);
-		if (OAPV_FAILED(Result))
+		TArray<ElectraDecodersUtil::APVVideo::FFramePBUInfo> FramePBUs;
+		ElectraDecodersUtil::APVVideo::ParseAUIntoFramePBUSubsamples(FramePBUs, MakeConstArrayView<uint8>(reinterpret_cast<const uint8*>(InInputAccessUnit.Data), InInputAccessUnit.DataSize));
+		// We only decode the first frame PBU, whatever it is.
+		if (FramePBUs.IsEmpty())
 		{
-			PostError(0, TEXT("APV decoding failed"), ERRCODE_INTERNAL_FAILED_TO_DECODE_INPUT);
+			PostError(0, TEXT("Failed to parse APV access unit for frame PBUs"), ERRCODE_INTERNAL_FAILED_TO_DECODE_INPUT);
 			return IElectraDecoder::EDecoderError::Error;
 		}
-		// Find a frame we can decode
-		const oapv_frm_info_t* FrameToDecode = nullptr;
-		for(int32 frmN=0; frmN<AccessUnitInfo.num_frms; ++frmN)
+		// Parse the frame information.
+		ElectraDecodersUtil::APVVideo::FFrameInfo FrameInfo;
+		if (!ElectraDecodersUtil::APVVideo::ParseFrameHeader(FrameInfo, MakeConstArrayView<uint8>(reinterpret_cast<const uint8*>(InInputAccessUnit.Data) + 4 + FramePBUs[0].PBUOffset, FramePBUs[0].PBUSize - 4)))
 		{
-			if (AccessUnitInfo.frm_info[frmN].pbu_type == OAPV_PBU_TYPE_PRIMARY_FRAME &&
-				AccessUnitInfo.frm_info[frmN].chroma_format_idc == 2 &&
-				AccessUnitInfo.frm_info[frmN].bit_depth == 10)
-			{
-				FrameToDecode = &AccessUnitInfo.frm_info[frmN];
-				break;
-			}
+			PostError(0, TEXT("Failed to parse first APV frame PBU in access unit"), ERRCODE_INTERNAL_FAILED_TO_DECODE_INPUT);
+			return IElectraDecoder::EDecoderError::Error;
 		}
-		if (!FrameToDecode)
+		if (FrameInfo.BitDepthMinus8 != 2 || FrameInfo.ChromaFormatIDC != 2)
 		{
-			PostError(0, TEXT("Could not locate a decodable APV frame in access unit"), ERRCODE_INTERNAL_FAILED_TO_DECODE_INPUT);
+			PostError(0, TEXT("Unsupported APV frame format (bitdepth or chroma format)"), ERRCODE_INTERNAL_FAILED_TO_DECODE_INPUT);
 			return IElectraDecoder::EDecoderError::Error;
 		}
-		check(2 == OAPV_CS_GET_BYTE_DEPTH(FrameToDecode->cs));
-		check(0 == OAPV_CS_GET_ENDIAN(FrameToDecode->cs));
 
 		oapv_frms_t OutFrames;
 		FMemory::Memzero(OutFrames);
@@ -509,14 +507,14 @@ IElectraDecoder::EDecoderError FVideoDecoderAPVElectra::DecodeAccessUnit(const F
 		int ColorSpace = OAPV_CS_SET(OAPV_CF_PLANAR2, 10, 0);
 
 		TSharedPtr<TArray64<uint8>, ESPMode::ThreadSafe> DecodedP010Buffer = MakeShared<TArray64<uint8>, ESPMode::ThreadSafe>();
-		int32 OutputW = Align(FrameToDecode->w, OAPV_MB_W);
-		int32 OutputH = Align(FrameToDecode->h, OAPV_MB_H);
+		int32 OutputW = Align(FrameInfo.FrameWidth, OAPV_MB_W);
+		int32 OutputH = Align(FrameInfo.FrameHeight, OAPV_MB_H);
 		DecodedP010Buffer->AddUninitialized(OutputW * OutputH * sizeof(uint16) * 2);
 
 		oapv_imgb_t ib;
 		FMemory::Memzero(ib);
-		ib.w[0] = FrameToDecode->w;
-		ib.h[0] = FrameToDecode->h;
+		ib.w[0] = FrameInfo.FrameWidth;
+		ib.h[0] = FrameInfo.FrameHeight;
 		int bd = OAPV_CS_GET_BYTE_DEPTH(ColorSpace);
 		// Setup for 2 planes
 		ib.w[1] = ib.w[0];
@@ -543,13 +541,22 @@ IElectraDecoder::EDecoderError FVideoDecoderAPVElectra::DecodeAccessUnit(const F
 
 		oapv_bitb_t InputBS;
 		FMemory::Memzero(InputBS);
-		InputBS.addr = pAU;
-		InputBS.ssize = nAU;
+		// The decoder checks for the `aPv1` magic which may not be present if we decode one sub frame from
+		// the access unit. We insert this magic, thereby modifying the bitstream which is why we are signalling
+		// this when creating the bitstream processor.
+		if (FramePBUs[0].PBUOffset < 4)
+		{
+			PostError(0, TEXT("APV decoding failed, bad frame PBU offset"), ERRCODE_INTERNAL_FAILED_TO_DECODE_INPUT);
+			return IElectraDecoder::EDecoderError::Error;
+		}
+		InputBS.addr = (void*)(reinterpret_cast<const uint8*>(InInputAccessUnit.Data) + FramePBUs[0].PBUOffset - 4);
+		InputBS.ssize = FramePBUs[0].PBUSize + 4;
+		*reinterpret_cast<uint32*>(InputBS.addr) = Electra::EndianSwap((uint32)0x61507631U);
 
 		oapvd_stat_t Stats;
 		FMemory::Memzero(Stats);
 
-        Result = oapvd_decode(DecoderInstance, &InputBS, &OutFrames, DecoderMetaInstance, &Stats);
+        int Result = oapvd_decode(DecoderInstance, &InputBS, &OutFrames, DecoderMetaInstance, &Stats);
 
 		if (OAPV_FAILED(Result))
 		{
@@ -585,8 +592,8 @@ IElectraDecoder::EDecoderError FVideoDecoderAPVElectra::DecodeAccessUnit(const F
 			}
 		}
 
-		int32 DisplayWidth = (int32)FrameToDecode->w;
-		int32 DisplayHeight = (int32)FrameToDecode->h;
+		int32 DisplayWidth = (int32)FrameInfo.FrameWidth;
+		int32 DisplayHeight = (int32)FrameInfo.FrameHeight;
 
 		TSharedPtr<FVideoDecoderOutputAPVElectra, ESPMode::ThreadSafe> NewOutput = MakeShared<FVideoDecoderOutputAPVElectra>();
 		NewOutput->PTS = InInputAccessUnit.PTS;
diff --git a/Engine/Plugins/Media/ElectraCodecs/Source/ElectraDecoders/Private/Utils/OpenAPV/ElectraBitstreamProcessor_APV.cpp b/Engine/Plugins/Media/ElectraCodecs/Source/ElectraDecoders/Private/Utils/OpenAPV/ElectraBitstreamProcessor_APV.cpp
new file mode 100644
--- /dev/null
+++ b/Engine/Plugins/Media/ElectraCodecs/Source/ElectraDecoders/Private/Utils/OpenAPV/ElectraBitstreamProcessor_APV.cpp
@@ -0,0 +1,179 @@
+// Copyright Epic Games, Inc. All Rights Reserved.
+
+#include "Utils/OpenAPV/ElectraBitstreamProcessor_APV.h"
+//#include "Utils/OpenAPV/ElectraUtilsAPVVideo.h"
+#include "Utils/MPEG/ElectraUtilsMPEGVideo.h"
+#include "Utilities/UtilitiesMP4.h"
+#include "IElectraDecoderFeaturesAndOptions.h"
+#include "ElectraDecodersUtils.h"
+
+
+class FElectraDecoderBitstreamProcessorAPV::FImpl
+{
+public:
+	class FBitstreamInfo : public IElectraDecoderBitstreamInfo
+	{
+	public:
+		virtual ~FBitstreamInfo() = default;
+	};
+
+	void SetPropertiesOnOutput(TMap<FString, FVariant>& InOutProperties)
+	{
+		if (CurrentColorimetry.IsValid() && bColorimetryChanged)
+		{
+			bColorimetryChanged = false;
+			TConstArrayView<uint8> Colorimetry = MakeConstArrayView(reinterpret_cast<const uint8*>(CurrentColorimetry.Get()), sizeof(ElectraDecodersUtil::MPEG::FCommonColorimetry));
+			InOutProperties.Emplace(IElectraDecoderBitstreamProcessorInfo::CommonColorimetry, TArray<uint8>(Colorimetry));
+		}
+
+		if (CurrentMDCV.IsValid() && bMDCVChanged)
+		{
+			bMDCVChanged = false;
+			TConstArrayView<uint8> Params = MakeConstArrayView(reinterpret_cast<const uint8*>(CurrentMDCV.Get()), sizeof(ElectraDecodersUtil::MPEG::FSEImastering_display_colour_volume));
+			InOutProperties.Emplace(IElectraDecoderBitstreamProcessorInfo::SeiMasteringDisplayColorVolume, TArray<uint8>(Params));
+		}
+
+		if (CurrentCLLI.IsValid() && bCLLIChanged)
+		{
+			bCLLIChanged = false;
+			TConstArrayView<uint8> Params = MakeConstArrayView(reinterpret_cast<const uint8*>(CurrentCLLI.Get()), sizeof(ElectraDecodersUtil::MPEG::FSEIcontent_light_level_info));
+			InOutProperties.Emplace(IElectraDecoderBitstreamProcessorInfo::SeiContentLightLeveInfo, TArray<uint8>(Params));
+		}
+	}
+
+	void UpdateColorimetry(uint8 colour_primaries, uint8 transfer_characteristics, uint8 matrix_coeffs, uint8 video_full_range_flag, uint8 video_format)
+	{
+		if (!CurrentColorimetry.IsValid() ||
+			CurrentColorimetry->colour_primaries != colour_primaries ||
+			CurrentColorimetry->transfer_characteristics != transfer_characteristics ||
+			CurrentColorimetry->matrix_coeffs != matrix_coeffs ||
+			CurrentColorimetry->video_full_range_flag != video_full_range_flag ||
+			CurrentColorimetry->video_format != video_format)
+		{
+			CurrentColorimetry = MakeShared<ElectraDecodersUtil::MPEG::FCommonColorimetry, ESPMode::ThreadSafe>();
+			CurrentColorimetry->colour_primaries = colour_primaries;
+			CurrentColorimetry->transfer_characteristics = transfer_characteristics;
+			CurrentColorimetry->matrix_coeffs = matrix_coeffs;
+			CurrentColorimetry->video_full_range_flag = video_full_range_flag;
+			CurrentColorimetry->video_format = video_format;
+			bColorimetryChanged = true;
+		}
+	}
+
+	void Clear()
+	{
+		CurrentColorimetry.Reset();
+		CurrentMDCV.Reset();
+		CurrentCLLI.Reset();
+		bColorimetryChanged = false;
+		bMDCVChanged = false;
+		bCLLIChanged = false;
+		LastErrorMessage.Empty();
+	}
+
+	static FString DCRName;
+	static FString CSDName;
+
+	TSharedPtr<ElectraDecodersUtil::MPEG::FCommonColorimetry, ESPMode::ThreadSafe> CurrentColorimetry;
+	TSharedPtr<ElectraDecodersUtil::MPEG::FSEImastering_display_colour_volume, ESPMode::ThreadSafe> CurrentMDCV;
+	TSharedPtr<ElectraDecodersUtil::MPEG::FSEIcontent_light_level_info, ESPMode::ThreadSafe> CurrentCLLI;
+
+	bool bModifiesInputAUInPlace = false;
+
+	bool bColorimetryChanged = false;
+	bool bMDCVChanged = false;
+	bool bCLLIChanged = false;
+
+	FString LastErrorMessage;
+};
+FString FElectraDecoderBitstreamProcessorAPV::FImpl::DCRName(TEXT("dcr"));
+FString FElectraDecoderBitstreamProcessorAPV::FImpl::CSDName(TEXT("csd"));
+
+
+FElectraDecoderBitstreamProcessorAPV::~FElectraDecoderBitstreamProcessorAPV()
+{
+}
+
+FElectraDecoderBitstreamProcessorAPV::FElectraDecoderBitstreamProcessorAPV(const TMap<FString, FVariant>& InDecoderParams, const TMap<FString, FVariant>& InFormatParams)
+{
+	Impl = MakePimpl<FImpl>();
+	Impl->bModifiesInputAUInPlace = !!ElectraDecodersUtil::GetVariantValueSafeI64(InDecoderParams, IElectraDecoderFeature::ModifiesInputAUInPlace, 0);
+}
+
+bool FElectraDecoderBitstreamProcessorAPV::WillModifyBitstreamInPlace()
+{
+	return Impl->bModifiesInputAUInPlace;
+}
+
+void FElectraDecoderBitstreamProcessorAPV::Clear()
+{
+	Impl->Clear();
+}
+
+IElectraDecoderBitstreamProcessor::EProcessResult FElectraDecoderBitstreamProcessorAPV::GetCSDFromConfigurationRecord(TArray<uint8>& OutCSD, const TMap<FString, FVariant>& InParamsWithDCRorCSD)
+{
+	OutCSD.Empty();
+	return EProcessResult::Ok;
+}
+
+IElectraDecoderBitstreamProcessor::EProcessResult FElectraDecoderBitstreamProcessorAPV::ProcessInputForDecoding(TSharedPtr<IElectraDecoderBitstreamInfo, ESPMode::ThreadSafe>& OutBSI, FElectraDecoderInputAccessUnit& InOutAccessUnit, const TMap<FString, FVariant>& InAccessUnitSidebandData)
+{
+	// Already processed?
+	if ((InOutAccessUnit.Flags & EElectraDecoderFlags::InputIsProcessed) != EElectraDecoderFlags::None)
+	{
+		return EProcessResult::Ok;
+	}
+
+	// Set to processed even if we fail somewhere now.
+	InOutAccessUnit.Flags |= EElectraDecoderFlags::InputIsProcessed;
+	// Every frame is a keyframe
+	InOutAccessUnit.Flags |= EElectraDecoderFlags::IsSyncSample;
+
+#if 0
+	// Go over the data and look at the PBUs
+	const uint8* CurrentPBU = reinterpret_cast<const uint8*>(InOutAccessUnit.Data);
+	const uint8* EndOfPBUs = CurrentPBU + InOutAccessUnit.DataSize;
+	// Skip over size and signature if present.
+	if (InOutAccessUnit.DataSize >= 8 && Electra::GetFromBigEndian(reinterpret_cast<const uint32*>(InOutAccessUnit.Data)[1]) == 0x61507631)
+	{
+		CurrentPBU += 8;
+	}
+
+	while(CurrentPBU < EndOfPBUs)
+	{
+		const uint32 PBUsize = Electra::GetFromBigEndian(*reinterpret_cast<const uint32*>(CurrentPBU));
+		switch(CurrentPBU[4])
+		{
+			case 1:		// primary frame
+			case 2:		// non-primary frame
+			case 25:	// preview frame
+			case 26:	// depth frame
+			case 27:	// alpha frame
+			case 65:	// access unit information
+			case 67:	// filler
+			default:
+			{
+				break;
+			}
+			case 66:	// metadata
+			{
+				break;
+			}
+		}
+		CurrentPBU += PBUsize + 4;
+	}
+#endif
+	return EProcessResult::Ok;
+}
+
+void FElectraDecoderBitstreamProcessorAPV::SetPropertiesOnOutput(TMap<FString, FVariant>& InOutProperties, TSharedPtr<IElectraDecoderBitstreamInfo, ESPMode::ThreadSafe> InBSI)
+{
+	Impl->SetPropertiesOnOutput(InOutProperties);
+//	uint8 num_bits = 8;
+//	InOutProperties.Emplace(IElectraDecoderBitstreamProcessorInfo::NumBitsLuma, FVariant(num_bits));
+}
+
+FString FElectraDecoderBitstreamProcessorAPV::GetLastError()
+{
+	return Impl->LastErrorMessage;
+}
diff --git a/Engine/Plugins/Media/ElectraCodecs/Source/ElectraDecoders/Private/Utils/OpenAPV/ElectraUtilsAPVVideo.cpp b/Engine/Plugins/Media/ElectraCodecs/Source/ElectraDecoders/Private/Utils/OpenAPV/ElectraUtilsAPVVideo.cpp
new file mode 100644
--- /dev/null
+++ b/Engine/Plugins/Media/ElectraCodecs/Source/ElectraDecoders/Private/Utils/OpenAPV/ElectraUtilsAPVVideo.cpp
@@ -0,0 +1,375 @@
+// Copyright Epic Games, Inc. All Rights Reserved.
+
+#include "Utils/OpenAPV/ElectraUtilsAPVVideo.h"
+#include "Utilities/ElectraBitstream.h"
+#include "Utilities/UtilitiesMP4.h"
+
+
+namespace ElectraDecodersUtil
+{
+	namespace APVVideo
+	{
+
+		bool ParseFrameHeader(FFrameInfo& OutFrameInfo, const TConstArrayView<uint8> InFramePBU)
+		{
+			Electra::FBitstreamReader bs(InFramePBU.GetData(), InFramePBU.Num(), 4);	// skip over pbu type, group id and reserved
+			if (bs.GetRemainingByteLength() < 14)
+			{
+				return false;
+			}
+			// frame_info
+			OutFrameInfo.ProfileIDC = bs.GetBits(8);
+			OutFrameInfo.LevelIDC = bs.GetBits(8);
+			OutFrameInfo.BandIDC = bs.GetBits(3);
+			bs.SkipBits(5);
+			OutFrameInfo.FrameWidth = bs.GetBits(24);
+			OutFrameInfo.FrameHeight = bs.GetBits(24);
+			OutFrameInfo.ChromaFormatIDC = bs.GetBits(4);
+			OutFrameInfo.BitDepthMinus8 = bs.GetBits(4);
+			OutFrameInfo.CaptureTimeDistance = bs.GetBits(8);
+			bs.SkipBits(8);
+			// rest
+			bs.SkipBits(8);
+			OutFrameInfo.bHaveColorDescription = !!bs.GetBits(1);
+			if (OutFrameInfo.bHaveColorDescription)
+			{
+				if (bs.GetRemainingBits() < 25)
+				{
+					return false;
+				}
+				OutFrameInfo.ColorInfo.ColorPrimaries = bs.GetBits(8);
+				OutFrameInfo.ColorInfo.TransferCharacteristics = bs.GetBits(8);
+				OutFrameInfo.ColorInfo.MatrixCoefficients = bs.GetBits(8);
+				OutFrameInfo.ColorInfo.VideoFullRangeFlag = bs.GetBits(1);
+			}
+			// the remainder is of no interest to us.
+			return true;
+		}
+
+
+		FString FAPVDecoderConfigurationRecord::GetFormatInfo(uint32 InProfile, uint32 InLevel)
+		{
+			FString fi;
+			if (InProfile == 33)
+			{
+				fi = TEXT("422-10");
+			}
+			else if (InProfile == 44)
+			{
+				fi = TEXT("422-12");
+			}
+			else if (InProfile == 55)
+			{
+				fi = TEXT("444-10");
+			}
+			else if (InProfile == 66)
+			{
+				fi = TEXT("444-12");
+			}
+			else if (InProfile == 77)
+			{
+				fi = TEXT("4444-10");
+			}
+			else if (InProfile == 88)
+			{
+				fi = TEXT("4444-12");
+			}
+			else if (InProfile == 99)
+			{
+				fi = TEXT("400-10");
+			}
+			else
+			{
+				fi = TEXT("Unknown");
+			}
+			fi.Append(TEXT(" Profile"));
+			if (InLevel >= 30 && InLevel < 240)
+			{
+				fi.Append(FString::Printf(TEXT(", level %d.%d"), InLevel/30 , InLevel%30));
+			}
+			else
+			{
+				fi.Append(TEXT(" at unknown level"));
+			}
+			return fi;
+		}
+
+		FString FAPVDecoderConfigurationRecord::GetFormatInfo()
+		{
+			return FString::Printf(TEXT("APV, %s"), *GetFormatInfo(Profile, Level));
+		}
+
+
+		bool FAPVDecoderConfigurationRecord::Parse(const TConstArrayView<uint8> InDCR)
+		{
+			RawData.Empty();
+			ColorInfo = FColorInfo();
+			Width = 0;
+			Height = 0;
+			Profile = 0;
+			Level = 0;
+			Band = 0;
+			BitDepth = 0;
+			LargestProfile = 0;
+			LargestLevel = 0;
+			LargestBand = 0;
+			NumEncodedFrameVariants = 0;
+			bIsNonConformant = false;
+
+			if (InDCR.Num() < 2 || InDCR[0] != 1)
+			{
+				return false;
+			}
+			RawData = InDCR;
+			const uint8* InapvCBox = InDCR.GetData() + 1;
+			const uint8* End = InDCR.GetData() + InDCR.Num();
+			const int32 NumCfgs = *InapvCBox++;
+			bool bHavePrimary = false;
+			for(int32 i=0; i<NumCfgs && InapvCBox+1<End; ++i)
+			{
+				FConfigurationEntry e;
+				e.PBUType = *InapvCBox++;
+				const int32 NumFi = *InapvCBox++;
+				for(int32 j=0; j<NumFi; ++j)
+				{
+					// Have enough remaining data?
+					if (InapvCBox+13 >= End)
+					{
+						return false;
+					}
+					FFrameInfo fi;
+					fi.bHaveColorDescription = (*InapvCBox & 2) != 0;
+					fi.bCaptureTimeDistanceIgnored = (*InapvCBox & 1) != 0;
+					++InapvCBox;
+					fi.ProfileIDC = *InapvCBox++;
+					fi.LevelIDC = *InapvCBox++;
+					fi.BandIDC = *InapvCBox++;
+					fi.FrameWidth = Electra::GetFromBigEndian(*reinterpret_cast<const uint32*>(InapvCBox));
+					InapvCBox += 4;
+					fi.FrameHeight = Electra::GetFromBigEndian(*reinterpret_cast<const uint32*>(InapvCBox));
+					InapvCBox += 4;
+					fi.ChromaFormatIDC = *InapvCBox >> 4;
+					fi.BitDepthMinus8 = *InapvCBox++ & 15;
+					fi.CaptureTimeDistance = *InapvCBox++;
+					if (fi.bHaveColorDescription)
+					{
+						// Have enough remaining data?
+						if (InapvCBox+3 >= End)
+						{
+							return false;
+						}
+						fi.ColorInfo.ColorPrimaries = *InapvCBox++;
+						fi.ColorInfo.TransferCharacteristics = *InapvCBox++;
+						fi.ColorInfo.MatrixCoefficients = *InapvCBox++;
+						fi.ColorInfo.VideoFullRangeFlag = *InapvCBox++ >> 7;
+					}
+					// Keep track of the largest of profile, level and band for the RFC 6381 sub parameters
+					LargestProfile = fi.ProfileIDC > LargestProfile ? fi.ProfileIDC : LargestProfile;
+					LargestLevel = fi.LevelIDC > LargestLevel ? fi.LevelIDC : LargestLevel;
+					LargestBand = fi.BandIDC > LargestBand ? fi.BandIDC : LargestBand;
+					e.FrameInfos.Emplace(fi);
+					++NumEncodedFrameVariants;
+				}
+				// Is this the primary frame?
+				if (e.PBUType == 1)
+				{
+					bHavePrimary = true;
+					// Set the convenience values from the primary frame
+					Width = e.FrameInfos[0].FrameWidth;
+					Height = e.FrameInfos[0].FrameHeight;
+					ColorInfo = e.FrameInfos[0].ColorInfo;
+					Profile = e.FrameInfos[0].ProfileIDC;
+					Level = e.FrameInfos[0].LevelIDC;
+					Band = e.FrameInfos[0].BandIDC;
+					BitDepth = 8 + e.FrameInfos[0].BitDepthMinus8;
+					// There can be at most one primary frame.
+					bIsNonConformant = bIsNonConformant | (e.FrameInfos.Num() > 1);
+				}
+				ConfigurationEntries.Emplace(MoveTemp(e));
+			}
+			if (ConfigurationEntries.IsEmpty())
+			{
+				return false;
+			}
+			// A primary frame is required.
+			bIsNonConformant |= !bHavePrimary;
+			return true;
+		}
+
+		bool FAPVDecoderConfigurationRecord::ParseFromAU(const TConstArrayView64<uint8> InAccessUnit)
+		{
+			struct FPBU
+			{
+				uint8 Type;
+				TConstArrayView<uint8> Data;
+			};
+			RawData.Empty();
+			ColorInfo = FColorInfo();
+			Width = 0;
+			Height = 0;
+			Profile = 0;
+			Level = 0;
+			Band = 0;
+			BitDepth = 0;
+			LargestProfile = 0;
+			LargestLevel = 0;
+			LargestBand = 0;
+			NumEncodedFrameVariants = 0;
+			bIsNonConformant = false;
+
+			const uint8* CurrentPBU = reinterpret_cast<const uint8*>(InAccessUnit.GetData());
+			const uint8* EndOfPBUs = CurrentPBU + InAccessUnit.Num();
+			bool bHavePrimary = false;
+			// Skip over size and signature if present.
+			if (InAccessUnit.Num() >= 8 && Electra::GetFromBigEndian(reinterpret_cast<const uint32*>(CurrentPBU)[1]) == 0x61507631)
+			{
+				CurrentPBU += 8;
+			}
+			TMap<uint8, TArray<FPBU>> PBUs;
+			while(CurrentPBU < EndOfPBUs)
+			{
+				const uint32 PBUsize = Electra::GetFromBigEndian(*reinterpret_cast<const uint32*>(CurrentPBU));
+				FPBU pbu;
+				pbu.Type = CurrentPBU[4];
+				pbu.Data = MakeConstArrayView(CurrentPBU + 4, PBUsize);
+				CurrentPBU += PBUsize + 4;
+				TArray<FPBU>& pa = PBUs.FindOrAdd(pbu.Type);
+				pa.Emplace(pbu);
+			}
+
+			// Parse the values from a frame.
+			auto ParseFrameType = [&](uint8 InType) -> bool
+			{
+				if (PBUs.Contains(InType))
+				{
+					FConfigurationEntry* cfgE = nullptr;
+					const TArray<FPBU>& PBUArr(PBUs[InType]);
+					for(int32 i=0; i<PBUArr.Num(); ++i)
+					{
+						FFrameInfo fi;
+						if (ParseFrameHeader(fi, PBUArr[i].Data))
+						{
+							cfgE = ConfigurationEntries.FindByPredicate([&](const FConfigurationEntry& ce){return ce.PBUType == InType;});
+							if (!cfgE)
+							{
+								ConfigurationEntries.Emplace_GetRef().PBUType = InType;
+								cfgE = &ConfigurationEntries.Last();
+							}
+							// Keep track of the largest of profile, level and band for the RFC 6381 sub parameters
+							LargestProfile = fi.ProfileIDC > LargestProfile ? fi.ProfileIDC : LargestProfile;
+							LargestLevel = fi.LevelIDC > LargestLevel ? fi.LevelIDC : LargestLevel;
+							LargestBand = fi.BandIDC > LargestBand ? fi.BandIDC : LargestBand;
+							cfgE->FrameInfos.Emplace(fi);
+							++NumEncodedFrameVariants;
+						}
+						else
+						{
+							return false;
+						}
+					}
+					if (cfgE && InType == 1)
+					{
+						bHavePrimary = true;
+						// Set the convenience values from the primary frame
+						Width = cfgE->FrameInfos[0].FrameWidth;
+						Height = cfgE->FrameInfos[0].FrameHeight;
+						ColorInfo = cfgE->FrameInfos[0].ColorInfo;
+						Profile = cfgE->FrameInfos[0].ProfileIDC;
+						Level = cfgE->FrameInfos[0].LevelIDC;
+						Band = cfgE->FrameInfos[0].BandIDC;
+						BitDepth = 8 + cfgE->FrameInfos[0].BitDepthMinus8;
+						// There can be at most one primary frame.
+						bIsNonConformant = bIsNonConformant | (cfgE->FrameInfos.Num() > 1);
+					}
+				}
+				return true;
+			};
+			// Check all type of frames
+			if (!ParseFrameType(1) || !ParseFrameType(2) || !ParseFrameType(25) || !ParseFrameType(26) || !ParseFrameType(27))
+			{
+				return false;
+			}
+			if (ConfigurationEntries.IsEmpty())
+			{
+				return false;
+			}
+			// A primary frame is required.
+			bIsNonConformant |= !bHavePrimary;
+			// Build the raw data
+			Electra::FBitstreamWriter wr;
+			wr.PutBits(1U, 8);
+			wr.PutBits((uint32) ConfigurationEntries.Num(), 8);
+			for(int32 i=0; i<ConfigurationEntries.Num(); ++i)
+			{
+				const FConfigurationEntry& e(ConfigurationEntries[i]);
+				wr.PutBits(e.PBUType, 8);
+				wr.PutBits((uint32)e.FrameInfos.Num(), 8);
+				for(int32 j=0; j<e.FrameInfos.Num(); ++j)
+				{
+					const FFrameInfo& fi(e.FrameInfos[j]);
+					wr.PutBits(0U, 6);
+					wr.PutBits(fi.bHaveColorDescription?1U:0U, 1);
+					wr.PutBits(fi.bCaptureTimeDistanceIgnored?1U:0U, 1);
+					wr.PutBits(fi.ProfileIDC, 8);
+					wr.PutBits(fi.LevelIDC, 8);
+					wr.PutBits(fi.BandIDC, 8);
+					wr.PutBits(fi.FrameWidth, 32);
+					wr.PutBits(fi.FrameHeight, 32);
+					wr.PutBits(fi.ChromaFormatIDC, 4);
+					wr.PutBits(fi.BitDepthMinus8, 4);
+					wr.PutBits(fi.CaptureTimeDistance, 8);
+					if (fi.bHaveColorDescription)
+					{
+						wr.PutBits(fi.ColorInfo.ColorPrimaries, 8);
+						wr.PutBits(fi.ColorInfo.TransferCharacteristics, 8);
+						wr.PutBits(fi.ColorInfo.MatrixCoefficients, 8);
+						wr.PutBits(fi.ColorInfo.VideoFullRangeFlag, 1);
+						wr.PutBits(0U, 7);
+					}
+				}
+			}
+			wr.GetArray(RawData);
+			return true;
+		}
+
+
+		void ParseAUIntoFramePBUSubsamples(TArray<FFramePBUInfo>& OutPBUInfos, const TConstArrayView64<uint8> InAccessUnit)
+		{
+			const uint8* AUBase = reinterpret_cast<const uint8*>(InAccessUnit.GetData());
+			const uint8* CurrentPBU = AUBase;
+			const uint8* EndOfPBUs = CurrentPBU + InAccessUnit.Num();
+			// Skip over size and signature if present.
+			if (InAccessUnit.Num() >= 8 && Electra::GetFromBigEndian(reinterpret_cast<const uint32*>(CurrentPBU)[1]) == 0x61507631)
+			{
+				CurrentPBU += 8;
+			}
+			// If there is no size, skip over signature if present.
+			else if (InAccessUnit.Num() >= 4 && Electra::GetFromBigEndian(reinterpret_cast<const uint32*>(CurrentPBU)[0]) == 0x61507631)
+			{
+				CurrentPBU += 4;
+			}
+			while(CurrentPBU < EndOfPBUs)
+			{
+				const uint32 PBUsize = Electra::GetFromBigEndian(*reinterpret_cast<const uint32*>(CurrentPBU));
+				const uint8 Type = CurrentPBU[4];
+				// Any frame type PBU?
+				if (Type == 1 || Type == 2 || Type == 25 || Type == 26 || Type == 27)
+				{
+					if (!OutPBUInfos.IsEmpty())
+					{
+						OutPBUInfos.Last().PBUSize = (CurrentPBU - AUBase) - OutPBUInfos.Last().PBUOffset;
+					}
+					FFramePBUInfo& pi = OutPBUInfos.Emplace_GetRef();
+					pi.bIsPrimaryFrame = Type == 1;
+					pi.PBUOffset = CurrentPBU - AUBase;
+				}
+				CurrentPBU += PBUsize + 4;
+			}
+			if (!OutPBUInfos.IsEmpty())
+			{
+				OutPBUInfos.Last().PBUSize = (CurrentPBU - AUBase) - OutPBUInfos.Last().PBUOffset;
+			}
+		}
+
+	} // namespace APVVideo
+} // namespace ElectraDecodersUtil
diff --git a/Engine/Plugins/Media/ElectraCodecs/Source/ElectraDecoders/Public/Utils/OpenAPV/ElectraBitstreamProcessor_APV.h b/Engine/Plugins/Media/ElectraCodecs/Source/ElectraDecoders/Public/Utils/OpenAPV/ElectraBitstreamProcessor_APV.h
new file mode 100644
--- /dev/null
+++ b/Engine/Plugins/Media/ElectraCodecs/Source/ElectraDecoders/Public/Utils/OpenAPV/ElectraBitstreamProcessor_APV.h
@@ -0,0 +1,26 @@
+// Copyright Epic Games, Inc. All Rights Reserved.
+
+#pragma once
+
+#include "IElectraDecoder.h"
+
+
+class FElectraDecoderBitstreamProcessorAPV : public IElectraDecoderBitstreamProcessor
+{
+public:
+	static TSharedPtr<IElectraDecoderBitstreamProcessor, ESPMode::ThreadSafe> Create(const TMap<FString, FVariant>& InDecoderParams, const TMap<FString, FVariant>& InFormatParams)
+	{ return MakeShareable<>(new FElectraDecoderBitstreamProcessorAPV(InDecoderParams, InFormatParams)); }
+
+	virtual ~FElectraDecoderBitstreamProcessorAPV();
+	bool WillModifyBitstreamInPlace() override;
+	void Clear() override;
+	EProcessResult GetCSDFromConfigurationRecord(TArray<uint8>& OutCSD, const TMap<FString, FVariant>& InParamsWithDCRorCSD) override;
+	EProcessResult ProcessInputForDecoding(TSharedPtr<IElectraDecoderBitstreamInfo, ESPMode::ThreadSafe>& OutBSI, FElectraDecoderInputAccessUnit& InOutAccessUnit, const TMap<FString, FVariant>& InAccessUnitSidebandData) override;
+	void SetPropertiesOnOutput(TMap<FString, FVariant>& InOutProperties, TSharedPtr<IElectraDecoderBitstreamInfo, ESPMode::ThreadSafe> InBSI) override;
+	FString GetLastError() override;
+
+private:
+	ELECTRADECODERS_API FElectraDecoderBitstreamProcessorAPV(const TMap<FString, FVariant>& InDecoderParams, const TMap<FString, FVariant>& InFormatParams);
+	class FImpl;
+	TPimplPtr<FImpl> Impl;
+};
diff --git a/Engine/Plugins/Media/ElectraCodecs/Source/ElectraDecoders/Public/Utils/OpenAPV/ElectraUtilsAPVVideo.h b/Engine/Plugins/Media/ElectraCodecs/Source/ElectraDecoders/Public/Utils/OpenAPV/ElectraUtilsAPVVideo.h
new file mode 100644
--- /dev/null
+++ b/Engine/Plugins/Media/ElectraCodecs/Source/ElectraDecoders/Public/Utils/OpenAPV/ElectraUtilsAPVVideo.h
@@ -0,0 +1,116 @@
+// Copyright Epic Games, Inc. All Rights Reserved.
+
+#pragma once
+
+#include "CoreTypes.h"
+#include "Containers/Array.h"
+#include "Containers/ArrayView.h"
+#include "Containers/UnrealString.h"
+#include "Utilities/ElectraBitstream.h"
+
+namespace ElectraDecodersUtil
+{
+	namespace APVVideo
+	{
+		struct FColorInfo
+		{
+			uint8 ColorPrimaries = 2;
+			uint8 TransferCharacteristics = 2;
+			uint8 MatrixCoefficients = 2;
+			uint8 VideoFullRangeFlag = 0;
+		};
+		struct FFrameInfo
+		{
+			uint32 FrameWidth = 0;
+			uint32 FrameHeight = 0;
+			uint8 ProfileIDC = 0;
+			uint8 LevelIDC = 0;
+			uint8 BandIDC = 0;
+			uint8 ChromaFormatIDC = 0;
+			uint8 BitDepthMinus8 = 0;
+			uint8 CaptureTimeDistance = 0;
+			bool bHaveColorDescription = false;
+			bool bCaptureTimeDistanceIgnored = true;
+			FColorInfo ColorInfo;
+		};
+		ELECTRADECODERS_API bool ParseFrameHeader(FFrameInfo& OutFrameInfo, const TConstArrayView<uint8> InFramePBU);
+
+		/**
+		 * APV Codec ISO Media File Format Binding
+		 *
+		 * see: [REDACTED_PROFILE_URL]/openapv/blob/main/readme/apv_isobmff.md#apv-decoder-configuration-record-apvdecoderconfigurationrecord
+		 */
+		class FAPVDecoderConfigurationRecord
+		{
+		public:
+			struct FConfigurationEntry
+			{
+				uint8 PBUType = 0;
+				TArray<FFrameInfo> FrameInfos;
+			};
+
+			ELECTRADECODERS_API bool Parse(const TConstArrayView<uint8> InDCR);
+			ELECTRADECODERS_API bool ParseFromAU(const TConstArrayView64<uint8> InAccessUnit);
+
+			inline const TArray<uint8>& GetRawData() const
+			{ return RawData; }
+
+			inline int32 GetNumEncodedFrameVariants() const
+			{ return NumEncodedFrameVariants; }
+
+			inline uint32 GetWidth() const
+			{ return Width; }
+			inline uint32 GetHeight() const
+			{ return Height; }
+			inline uint8 GetProfile() const
+			{ return Profile; }
+			inline uint8 GetLevel() const
+			{ return Level; }
+			inline uint8 GetBand() const
+			{ return Band; }
+			inline uint8 GetBitDepth() const
+			{ return BitDepth; }
+			inline const FColorInfo& GetColorInfo() const
+			{ return ColorInfo; }
+			inline const TArray<FConfigurationEntry>& GetConfigurationEntries() const
+			{ return ConfigurationEntries; }
+
+			inline FString GetCodecSpecifierRFC6381() const
+			{
+				FString rfc(FString::Printf(TEXT("apv1.apvf%d.apvl%d.apvb%d"), LargestProfile, LargestLevel, LargestBand));
+				return rfc;
+			}
+
+			ELECTRADECODERS_API FString GetFormatInfo();
+
+			ELECTRADECODERS_API static FString GetFormatInfo(uint32 InProfile, uint32 InLevel);
+		private:
+			TArray<FConfigurationEntry> ConfigurationEntries;
+			TArray<uint8> RawData;
+			FColorInfo ColorInfo;
+			uint32 Width = 0;
+			uint32 Height = 0;
+			uint8 Profile = 0;
+			uint8 Level = 0;
+			uint8 Band = 0;
+			uint8 BitDepth = 0;
+			uint8 LargestProfile = 0;
+			uint8 LargestLevel = 0;
+			uint8 LargestBand = 0;
+			int32 NumEncodedFrameVariants = 0;
+			bool bIsNonConformant = false;
+		};
+
+
+		struct FFramePBUInfo
+		{
+			// Offset to where the size of the PBU is located
+			int64 PBUOffset = 0;
+			// Size of the PBU including the size field.
+			uint32 PBUSize = 0;
+			bool bIsPrimaryFrame = false;
+		};
+		ELECTRADECODERS_API void ParseAUIntoFramePBUSubsamples(TArray<FFramePBUInfo>& OutPBUInfos, const TConstArrayView64<uint8> InAccessUnit);
+
+	} // namespace APVVideo
+} // namespace ElectraDecodersUtil
