Allow H264 and H265 electra decoder reorder buffer sizes to be configured on Apple platforms

The frame reorder buffer was hardcoded to 5 frames for both decoder implementations. This change reads the "max_output_buffers" codec option which is used on Windows and Linux but ignored on Mac.

Enables low-latency streaming use cases to reduce decoder buffering when B-frames are not present.

#rnx

#Synced-CL 49653240

diff --git a/Engine/Plugins/Media/ElectraCodecs/Source/ElectraDecoders/Private/Apple/h264/VideoDecoderH264_Apple.cpp b/Engine/Plugins/Media/ElectraCodecs/Source/ElectraDecoders/Private/Apple/h264/VideoDecoderH264_Apple.cpp
--- a/Engine/Plugins/Media/ElectraCodecs/Source/ElectraDecoders/Private/Apple/h264/VideoDecoderH264_Apple.cpp
+++ b/Engine/Plugins/Media/ElectraCodecs/Source/ElectraDecoders/Private/Apple/h264/VideoDecoderH264_Apple.cpp
@@ -139,10 +139,11 @@ public:
 
 class FElectraVideoDecoderH264_Apple : public IElectraVideoDecoderH264_Apple
 {
-	// The decoder does not output frames in display order. We need to hold back a number of frames
-	// and sort them by PTS to get them in display order.
-	static constexpr int32 kNumReorderFrames = 5;
-
+	// The decoder may not output frames in display order depending on the presence of B-frames.
+	// We may need to hold back a number of frames and sort them by PTS to get them in display order.
+	// This value can be configured using the "max_output_buffers" additional option.
+	static constexpr int32 kDefaultNumReorderFrames = 5;
+	
 public:
 	FElectraVideoDecoderH264_Apple(const TMap<FString, FVariant>& InOptions);
 
@@ -339,6 +340,8 @@ private:
 
 	EDecodeState DecodeState = EDecodeState::Decoding;
 	bool bNewDecoderRequired = false;
+
+	int32 NumReorderFrames = kDefaultNumReorderFrames;
 };
 
 
@@ -392,6 +395,8 @@ TSharedPtr<IElectraDecoder, ESPMode::ThreadSafe> IElectraVideoDecoderH264_Apple:
 FElectraVideoDecoderH264_Apple::FElectraVideoDecoderH264_Apple(const TMap<FString, FVariant>& InOptions)
 {
 	InitialCreationOptions = InOptions;
+	const int32 MaxOutputBuffers = static_cast<int32>(ElectraDecodersUtil::GetVariantValueSafeU64(InOptions, TEXT("max_output_buffers"), kDefaultNumReorderFrames));
+	NumReorderFrames = FMath::Max(MaxOutputBuffers, 1);
 }
 
 FElectraVideoDecoderH264_Apple::~FElectraVideoDecoderH264_Apple()
@@ -855,7 +860,7 @@ bool FElectraVideoDecoderH264_Apple::CheckForAvailableOutput()
 	}
 
 	FScopeLock lock(&ReadyImageMutex);
-	bool bHaveEnough = ReadyImages.Num() >= kNumReorderFrames || DecodeState == EDecodeState::Draining;
+	bool bHaveEnough = ReadyImages.Num() >= NumReorderFrames || DecodeState == EDecodeState::Draining;
 	return bHaveEnough;
 }
 
diff --git a/Engine/Plugins/Media/ElectraCodecs/Source/ElectraDecoders/Private/Apple/h265/VideoDecoderH265_Apple.cpp b/Engine/Plugins/Media/ElectraCodecs/Source/ElectraDecoders/Private/Apple/h265/VideoDecoderH265_Apple.cpp
--- a/Engine/Plugins/Media/ElectraCodecs/Source/ElectraDecoders/Private/Apple/h265/VideoDecoderH265_Apple.cpp
+++ b/Engine/Plugins/Media/ElectraCodecs/Source/ElectraDecoders/Private/Apple/h265/VideoDecoderH265_Apple.cpp
@@ -139,9 +139,10 @@ public:
 
 class FElectraVideoDecoderH265_Apple : public IElectraVideoDecoderH265_Apple
 {
-	// The decoder does not output frames in display order. We need to hold back a number of frames
-	// and sort them by PTS to get them in display order.
-	static constexpr int32 kNumReorderFrames = 5;
+	// The decoder may not output frames in display order depending on the presence of B-frames.
+	// We may need to hold back a number of frames and sort them by PTS to get them in display order.
+	// This value can be configured using the "max_output_buffers" additional option.
+	static constexpr int32 kDefaultNumReorderFrames = 5;
 
 public:
 	FElectraVideoDecoderH265_Apple(const TMap<FString, FVariant>& InOptions);
@@ -339,6 +340,8 @@ private:
 
 	EDecodeState DecodeState = EDecodeState::Decoding;
 	bool bNewDecoderRequired = false;
+
+	int32 NumReorderFrames = kDefaultNumReorderFrames;
 };
 
 
@@ -392,6 +395,8 @@ TSharedPtr<IElectraDecoder, ESPMode::ThreadSafe> IElectraVideoDecoderH265_Apple:
 FElectraVideoDecoderH265_Apple::FElectraVideoDecoderH265_Apple(const TMap<FString, FVariant>& InOptions)
 {
 	InitialCreationOptions = InOptions;
+	const int32 MaxOutputBuffers = static_cast<int32>(ElectraDecodersUtil::GetVariantValueSafeU64(InOptions, TEXT("max_output_buffers"), kDefaultNumReorderFrames));
+	NumReorderFrames = FMath::Max(MaxOutputBuffers, 1);
 }
 
 FElectraVideoDecoderH265_Apple::~FElectraVideoDecoderH265_Apple()
@@ -858,7 +863,7 @@ bool FElectraVideoDecoderH265_Apple::CheckForAvailableOutput()
 	}
 
 	FScopeLock lock(&ReadyImageMutex);
-	bool bHaveEnough = ReadyImages.Num() >= kNumReorderFrames || DecodeState == EDecodeState::Draining;
+	bool bHaveEnough = ReadyImages.Num() >= NumReorderFrames || DecodeState == EDecodeState::Draining;
 	return bHaveEnough;
 }
 
