Use IMF2DBuffer2::Lock2DSize for Electra DX decoders

This allows us to sanity-check the buffer length.

#Synced-CL 50324810

diff --git a/Engine/Plugins/Media/ElectraCodecs/Source/ElectraDecoders/Private/DX/VideoDecoderH264_DX.cpp b/Engine/Plugins/Media/ElectraCodecs/Source/ElectraDecoders/Private/DX/VideoDecoderH264_DX.cpp
--- a/Engine/Plugins/Media/ElectraCodecs/Source/ElectraDecoders/Private/DX/VideoDecoderH264_DX.cpp
+++ b/Engine/Plugins/Media/ElectraCodecs/Source/ElectraDecoders/Private/DX/VideoDecoderH264_DX.cpp
@@ -970,8 +970,27 @@ void FElectraVideoDecoderH264_DX::CopyBufferData(TRefCountPtr<IMF2DBuffer> Buffe
 {
 	uint8* Data = nullptr;
 	LONG Pitch = 0;
-	HRESULT Res = Buffer2D->Lock2D(&Data, &Pitch);
-	check(SUCCEEDED(Res));
+	TRefCountPtr<IMF2DBuffer2> Buffer2D2;
+	HRESULT Res = Buffer2D->QueryInterface(IID_PPV_ARGS(Buffer2D2.GetInitReference()));
+	if (SUCCEEDED(Res))
+	{
+		uint8* Start = nullptr;
+		DWORD Length = 0;
+		Res = Buffer2D2->Lock2DSize(MF2DBuffer_LockFlags_Read, &Data, &Pitch, &Start, &Length);
+		if (SUCCEEDED(Res))
+		{
+			check(Length >= static_cast<DWORD>(Pitch * DecodedHeight * 3 / 2));
+		}
+		else
+		{
+			checkNoEntry();
+		}
+	}
+	else
+	{
+		Res = Buffer2D->Lock2D(&Data, &Pitch);
+		check(SUCCEEDED(Res));
+	}
 
 	ID3D12DynamicRHI* const D3D12RHI = GetID3D12DynamicRHI();
 	ID3D12Device* const Device = D3D12RHI->RHIGetDevice_NoMGPU();
diff --git a/Engine/Plugins/Media/ElectraCodecs/Source/ElectraDecoders/Private/DX/VideoDecoderH265_DX.cpp b/Engine/Plugins/Media/ElectraCodecs/Source/ElectraDecoders/Private/DX/VideoDecoderH265_DX.cpp
--- a/Engine/Plugins/Media/ElectraCodecs/Source/ElectraDecoders/Private/DX/VideoDecoderH265_DX.cpp
+++ b/Engine/Plugins/Media/ElectraCodecs/Source/ElectraDecoders/Private/DX/VideoDecoderH265_DX.cpp
@@ -955,8 +955,27 @@ void FElectraVideoDecoderH265_DX::CopyBufferData(TRefCountPtr<IMF2DBuffer> Buffe
 {
 	uint8* Data = nullptr;
 	LONG Pitch = 0;
-	HRESULT Res = Buffer2D->Lock2D(&Data, &Pitch);
-	check(SUCCEEDED(Res));
+	TRefCountPtr<IMF2DBuffer2> Buffer2D2;
+	HRESULT Res = Buffer2D->QueryInterface(IID_PPV_ARGS(Buffer2D2.GetInitReference()));
+	if (SUCCEEDED(Res))
+	{
+		uint8* Start = nullptr;
+		DWORD Length = 0;
+		Res = Buffer2D2->Lock2DSize(MF2DBuffer_LockFlags_Read, &Data, &Pitch, &Start, &Length);
+		if (SUCCEEDED(Res))
+		{
+			check(Length >= static_cast<DWORD>(Pitch * DecodedHeight * 3 / 2));
+		}
+		else
+		{
+			checkNoEntry();
+		}
+	}
+	else
+	{
+		Res = Buffer2D->Lock2D(&Data, &Pitch);
+		check(SUCCEEDED(Res));
+	}
 
 	ID3D12DynamicRHI* const D3D12RHI = GetID3D12DynamicRHI();
 	ID3D12Device* const Device = D3D12RHI->RHIGetDevice_NoMGPU();
