ElectraCodecs: Rebase timestamps in Android decoders as large values could cause issues

#rnx

#Synced-CL 51662825

diff --git a/Engine/Plugins/Media/ElectraCodecs/Source/ElectraDecoders/Private/Android/AndroidPlatformElectraDecoders.cpp b/Engine/Plugins/Media/ElectraCodecs/Source/ElectraDecoders/Private/Android/AndroidPlatformElectraDecoders.cpp
--- a/Engine/Plugins/Media/ElectraCodecs/Source/ElectraDecoders/Private/Android/AndroidPlatformElectraDecoders.cpp
+++ b/Engine/Plugins/Media/ElectraCodecs/Source/ElectraDecoders/Private/Android/AndroidPlatformElectraDecoders.cpp
@@ -8,9 +8,19 @@
 #include "vpx/VPx_VideoDecoder_Android.h"
 #include "aac/AAC_AudioDecoder_Android.h"
 
+#include "HAL/IConsoleManager.h"
+
+
 namespace PlatformElectraDecodersAndroid
 {
 	static TArray<TSharedPtr<IElectraCodecFactory, ESPMode::ThreadSafe>> Factories;
+
+	TAutoConsoleVariable<bool> CVarElectraDecoderAndroidRebasePTS(
+		TEXT("ElectraDecoders.AndroidRebasePTS"),
+		true,
+		TEXT("Rebase PTS to zero for Android decoders.\n")
+		TEXT(" true: rebase so values start at 0, false use original values"),
+		ECVF_Default);
 }
 
 
diff --git a/Engine/Plugins/Media/ElectraCodecs/Source/ElectraDecoders/Private/Android/aac/AudioDecoderAAC_Android.cpp b/Engine/Plugins/Media/ElectraCodecs/Source/ElectraDecoders/Private/Android/aac/AudioDecoderAAC_Android.cpp
--- a/Engine/Plugins/Media/ElectraCodecs/Source/ElectraDecoders/Private/Android/aac/AudioDecoderAAC_Android.cpp
+++ b/Engine/Plugins/Media/ElectraCodecs/Source/ElectraDecoders/Private/Android/aac/AudioDecoderAAC_Android.cpp
@@ -14,6 +14,14 @@
 
 #include "AudioDecoderAAC_JavaWrapper_Android.h"
 
+/*********************************************************************************************************************/
+
+#include "HAL/IConsoleManager.h"
+namespace PlatformElectraDecodersAndroid
+{
+extern TAutoConsoleVariable<bool> CVarElectraDecoderAndroidRebasePTS;
+}
+
 /*********************************************************************************************************************/
 #define ENABLE_DETAILED_LOG 0
 #if ENABLE_DETAILED_LOG
@@ -108,6 +116,7 @@ private:
 	{
 		FInputAccessUnit AccessUnit;
 		TMap<FString, FVariant> AdditionalOptions;
+		int64 DecoderPTS = -1;
 	};
 
 	enum class EDecodeState
@@ -133,6 +142,7 @@ private:
 	TSharedPtr<FElectraAudioDecoderOutputAAC_Android, ESPMode::ThreadSafe> CurrentOutput;
 
 	TSharedPtr<IElectraAACAudioDecoderAndroidJava, ESPMode::ThreadSafe> DecoderInstance;
+	TOptional<int64> DecoderPTSOffset;
 	int64 LastPushedPresentationTimeUs = 0;
 	bool bDidEnqueueInputBuffer = false;
 	IElectraAACAudioDecoderAndroidJava::FOutputFormatInfo CurrentOutputFormatInfo;
@@ -249,6 +259,7 @@ bool FElectraAudioDecoderAAC_Android::ResetToCleanStart()
 	DecodeState = EDecodeState::Decoding;
 	bDidSendEOS = false;
 	bIsOutputFormatInfoValid = false;
+	DecoderPTSOffset.Reset();
 	LastPushedPresentationTimeUs = 0;
 	bDidEnqueueInputBuffer = false;
 
@@ -311,17 +322,32 @@ IElectraDecoder::EDecoderError FElectraAudioDecoderAAC_Android::DecodeAccessUnit
 		int32 InputBufferIndex = DecoderInstance->DequeueInputBuffer(0);
 		if (InputBufferIndex >= 0)
 		{
-			DETAILLOG(LogElectraDecoders, Log, TEXT("AudioDecoderAAC::DecodeAccessUnit() - queue input %lld"), (long long int)InInputAccessUnit.PTS.GetTicks());
-			int32 Result = DecoderInstance->QueueInputBuffer(InputBufferIndex, InInputAccessUnit.Data, InInputAccessUnit.DataSize, InInputAccessUnit.PTS.GetTicks());
+			// The decoder could have issues when the PTS is large as is common in a Live stream.
+			// We rebase it to work around this issue.
+			int64 PTSTicks = InInputAccessUnit.PTS.GetTicks();
+			if (!DecoderPTSOffset.IsSet())
+			{
+				DecoderPTSOffset = PlatformElectraDecodersAndroid::CVarElectraDecoderAndroidRebasePTS.GetValueOnAnyThread() ? PTSTicks : 0;
+				DETAILLOG(LogElectraDecoders, Log, TEXT("AudioDecoderAAC::DecodeAccessUnit() - Rebase PTS offset %lld"), (long long int)DecoderPTSOffset.GetValue());
+			}
+			PTSTicks -= DecoderPTSOffset.GetValue();
+			// Also the decoder internally truncates nanoseconds to microseconds which it then
+			// sticks onto the decoded output, so we cannot use the PTS as a "user value".
+			// We truncate the value ourselves to account for this.
+			PTSTicks /= 10;		// Convert from HNS to microseconds
+
+			DETAILLOG(LogElectraDecoders, Log, TEXT("AudioDecoderAAC::DecodeAccessUnit() - queue input %lld (%lld)"), (long long int)PTSTicks, (long long int)InInputAccessUnit.PTS.GetTicks());
+			int32 Result = DecoderInstance->QueueInputBuffer(InputBufferIndex, InInputAccessUnit.Data, InInputAccessUnit.DataSize, PTSTicks);
 			if (Result == 0)
 			{
 				bDidEnqueueInputBuffer = true;
-				LastPushedPresentationTimeUs = InInputAccessUnit.PTS.GetTicks();
+				LastPushedPresentationTimeUs = PTSTicks;
 
 				// Add to the list of inputs passed to the decoder.
 				FDecoderInput In;
 				In.AdditionalOptions = InAdditionalOptions;
 				In.AccessUnit = InInputAccessUnit;
+				In.DecoderPTS = PTSTicks;
 				InDecoderInput.Emplace(MoveTemp(In));
 				InDecoderInput.Sort([](const FDecoderInput& e1, const FDecoderInput& e2)
 				{
@@ -399,6 +425,7 @@ IElectraDecoder::EDecoderError FElectraAudioDecoderAAC_Android::Flush()
 	DecodeState = EDecodeState::Decoding;
 	InDecoderInput.Empty();
 	CurrentOutput.Reset();
+	DecoderPTSOffset.Reset();
 	LastPushedPresentationTimeUs = 0;
 	bDidSendEOS = false;
 	bDidEnqueueInputBuffer = false;
diff --git a/Engine/Plugins/Media/ElectraCodecs/Source/ElectraDecoders/Private/Android/h264/VideoDecoderH264_Android.cpp b/Engine/Plugins/Media/ElectraCodecs/Source/ElectraDecoders/Private/Android/h264/VideoDecoderH264_Android.cpp
--- a/Engine/Plugins/Media/ElectraCodecs/Source/ElectraDecoders/Private/Android/h264/VideoDecoderH264_Android.cpp
+++ b/Engine/Plugins/Media/ElectraCodecs/Source/ElectraDecoders/Private/Android/h264/VideoDecoderH264_Android.cpp
@@ -18,6 +18,14 @@
 
 #include "h264/VideoDecoderH264_JavaWrapper_Android.h"
 
+/*********************************************************************************************************************/
+
+#include "HAL/IConsoleManager.h"
+namespace PlatformElectraDecodersAndroid
+{
+extern TAutoConsoleVariable<bool> CVarElectraDecoderAndroidRebasePTS;
+}
+
 /*********************************************************************************************************************/
 #define ENABLE_DETAILED_LOG 0
 #if ENABLE_DETAILED_LOG
@@ -167,6 +175,7 @@ private:
 		FInputAccessUnit AccessUnit;
 		TMap<FString, FVariant> AdditionalOptions;
 		TSharedPtr<ElectraDecodersUtil::MPEG::H264::FSequenceParameterSet, ESPMode::ThreadSafe> SPS;
+		int64 DecoderPTS = -1;
 	};
 
 	enum class EDecodeState
@@ -209,6 +218,7 @@ private:
 	TSharedPtr<IElectraH264VideoDecoderAndroidJava, ESPMode::ThreadSafe> DecoderInstance;
 	IElectraH264VideoDecoderAndroidJava::FDecoderInformation DecoderInfo;
 	EDecodeState DecodeState = EDecodeState::Decoding;
+	TOptional<int64> DecoderPTSOffset;
 	int64 LastPushedPresentationTime = -1;
 	bool bDidSendEOS = false;
 	bool bNewDecoderRequired = false;
@@ -374,6 +384,7 @@ bool FElectraVideoDecoderH264_Android::ResetToCleanStart()
 	CurrentOutput.Reset();
 	InDecoderInput.Empty();
 	DecodeState = EDecodeState::Decoding;
+	DecoderPTSOffset.Reset();
 	LastPushedPresentationTime = -1;
 	bDidSendEOS = false;
 
@@ -434,6 +445,7 @@ IElectraDecoder::EDecoderError FElectraVideoDecoderH264_Android::DecodeAccessUni
 		bTriggerReplay = true;
 		CurrentOutput.Reset();
 		InDecoderInput.Empty();
+		DecoderPTSOffset.Reset();
 		LastPushedPresentationTime = -1;
 		bNewDecoderRequired = true;
 	}
@@ -462,8 +474,20 @@ IElectraDecoder::EDecoderError FElectraVideoDecoderH264_Android::DecodeAccessUni
 		int32 InputBufferIndex = DecoderInstance->DequeueInputBuffer(0);
 		if (InputBufferIndex >= 0)
 		{
-			// Android decoders need the PTS in microseconds
-			int64 PTS = InInputAccessUnit.PTS.GetTicks() / 10;
+			// The decoder could have issues when the PTS is large as is common in a Live stream.
+			// We rebase it to work around this issue.
+			int64 PTSTicks = InInputAccessUnit.PTS.GetTicks();
+			if (!DecoderPTSOffset.IsSet())
+			{
+				DecoderPTSOffset = PlatformElectraDecodersAndroid::CVarElectraDecoderAndroidRebasePTS.GetValueOnAnyThread() ? PTSTicks : 0;
+				DETAILLOG(LogElectraDecoders, Log, TEXT("VideoDecoderH264::DecodeAccessUnit() - Rebase PTS offset %lld"), (long long int)DecoderPTSOffset.GetValue());
+			}
+			PTSTicks -= DecoderPTSOffset.GetValue();
+			// Also the decoder internally truncates nanoseconds to microseconds which it then
+			// sticks onto the decoded output, so we cannot use the PTS as a "user value".
+			// We truncate the value ourselves to account for this.
+			PTSTicks /= 10;		// Convert from HNS to microseconds
+
 			int32 Result = 0;
 			// If this is a sync sample we prepend the CSD. While this is not necessary on a running stream we need to have the CSD
 			// on the first frame and it is easier to prepend it to all IDR frames when seeking etc.
@@ -477,22 +501,24 @@ IElectraDecoder::EDecoderError FElectraVideoDecoderH264_Android::DecodeAccessUni
 				}
 
 				TArray<uint8> CSD = ElectraDecodersUtil::GetVariantValueUInt8Array(InAdditionalOptions, TEXT("csd"));
-				Result = DecoderInstance->QueueInputBuffer(InputBufferIndex, CSD.GetData(), CSD.Num(), InInputAccessUnit.Data, InInputAccessUnit.DataSize, PTS);
+				Result = DecoderInstance->QueueInputBuffer(InputBufferIndex, CSD.GetData(), CSD.Num(), InInputAccessUnit.Data, InInputAccessUnit.DataSize, PTSTicks);
 			}
 			else
 			{
-				Result = DecoderInstance->QueueInputBuffer(InputBufferIndex, nullptr, 0, InInputAccessUnit.Data, InInputAccessUnit.DataSize, PTS);
+				Result = DecoderInstance->QueueInputBuffer(InputBufferIndex, nullptr, 0, InInputAccessUnit.Data, InInputAccessUnit.DataSize, PTSTicks);
 			}
-			DETAILLOG(LogElectraDecoders, Log, TEXT("VideoDecoderH264::DecodeAccessUnit() - queue input %lld, 0x%x"), (long long int)InInputAccessUnit.PTS.GetTicks(), InInputAccessUnit.Flags);
+			DETAILLOG(LogElectraDecoders, Log, TEXT("VideoDecoderH264::DecodeAccessUnit() - queue input %lld (%lld), 0x%x"), (long long int)PTSTicks, (long long int)InInputAccessUnit.PTS.GetTicks(), InInputAccessUnit.Flags);
 			if (Result == 0)
 			{
-				LastPushedPresentationTime = PTS;
+				LastPushedPresentationTime = PTSTicks;
 
 				// Add to the list of inputs passed to the decoder.
 				TSharedPtr<FDecoderInput, ESPMode::NotThreadSafe> In(new FDecoderInput);
 				In->AdditionalOptions = InAdditionalOptions;
 				In->AccessUnit = InInputAccessUnit;
 				In->SPS = CurrentSPS;
+				In->DecoderPTS = PTSTicks;
+
 				// Zero the input pointer and size in the copy. That data is not owned by us and it's best not to have any
 				// values here that would lead us to think that we do.
 				In->AccessUnit.Data = nullptr;
@@ -573,6 +599,7 @@ IElectraDecoder::EDecoderError FElectraVideoDecoderH264_Android::Flush()
 	}
 	CurrentOutput.Reset();
 	InDecoderInput.Empty();
+	DecoderPTSOffset.Reset();
 	LastPushedPresentationTime = -1;
 	bDidSendEOS = false;
 	DecodeState = EDecodeState::Decoding;
@@ -915,8 +942,7 @@ FElectraVideoDecoderH264_Android::EConvertResult FElectraVideoDecoderH264_Androi
 	TSharedPtr<FDecoderInput, ESPMode::NotThreadSafe> In;
 	for(int32 nInDec=0; nInDec<InDecoderInput.Num(); ++nInDec)
 	{
-		// Decoder timestamps are in microseconds, so we need to compare as that.
-		if (InDecoderInput[nInDec]->AccessUnit.PTS / 10 == InInfo.PresentationTimestamp)
+		if (InDecoderInput[nInDec]->DecoderPTS == InInfo.PresentationTimestamp)
 		{
 			In = InDecoderInput[nInDec];
 			InDecoderInput.RemoveAt(nInDec);
@@ -937,14 +963,14 @@ FElectraVideoDecoderH264_Android::EConvertResult FElectraVideoDecoderH264_Androi
 	// If this is a replay frame or a frame that will not be output, we do not need to set it up and merely release the buffer immediately.
 	if ((In->AccessUnit.Flags & (EElectraDecoderFlags::IsReplaySample | EElectraDecoderFlags::DoNotOutput)) != EElectraDecoderFlags::None)
 	{
-		DETAILLOG(LogElectraDecoders, Log, TEXT("VideoDecoderH264::ConvertDecoderOutput() - Discard replay output %lld, 0x%x"), (long long int)In->AccessUnit.PTS.GetTicks(), In->AccessUnit.Flags);
+		DETAILLOG(LogElectraDecoders, Log, TEXT("VideoDecoderH264::ConvertDecoderOutput() - Discard output %lld (%lld), 0x%x"), (long long int)In->DecoderPTS, (long long int)In->AccessUnit.PTS.GetTicks(), In->AccessUnit.Flags);
 		if (DecoderInstance.IsValid())
 		{
 			DecoderInstance->ReleaseOutputBuffer(InInfo.BufferIndex, InInfo.ValidCount, false, -1);
 		}
 		return EConvertResult::Success;
 	}
-	DETAILLOG(LogElectraDecoders, Log, TEXT("VideoDecoderH264::ConvertDecoderOutput() - Convert output %lld, 0x%x"), (long long int)In->AccessUnit.PTS.GetTicks(), In->AccessUnit.Flags);
+	DETAILLOG(LogElectraDecoders, Log, TEXT("VideoDecoderH264::ConvertDecoderOutput() - Convert output %lld (%lld), 0x%x"), (long long int)In->DecoderPTS, (long long int)In->AccessUnit.PTS.GetTicks(), In->AccessUnit.Flags);
 
 	TSharedPtr<FElectraVideoDecoderOutputH264_Android, ESPMode::ThreadSafe> NewOutput = MakeShared<FElectraVideoDecoderOutputH264_Android>();
 	NewOutput->PTS = In->AccessUnit.PTS;
diff --git a/Engine/Plugins/Media/ElectraCodecs/Source/ElectraDecoders/Private/Android/h265/VideoDecoderH265_Android.cpp b/Engine/Plugins/Media/ElectraCodecs/Source/ElectraDecoders/Private/Android/h265/VideoDecoderH265_Android.cpp
--- a/Engine/Plugins/Media/ElectraCodecs/Source/ElectraDecoders/Private/Android/h265/VideoDecoderH265_Android.cpp
+++ b/Engine/Plugins/Media/ElectraCodecs/Source/ElectraDecoders/Private/Android/h265/VideoDecoderH265_Android.cpp
@@ -18,6 +18,14 @@
 
 #include "h265/VideoDecoderH265_JavaWrapper_Android.h"
 
+/*********************************************************************************************************************/
+
+#include "HAL/IConsoleManager.h"
+namespace PlatformElectraDecodersAndroid
+{
+extern TAutoConsoleVariable<bool> CVarElectraDecoderAndroidRebasePTS;
+}
+
 /*********************************************************************************************************************/
 #define ENABLE_DETAILED_LOG 0
 #if ENABLE_DETAILED_LOG
@@ -167,6 +175,7 @@ private:
 		FInputAccessUnit AccessUnit;
 		TMap<FString, FVariant> AdditionalOptions;
 		TSharedPtr<ElectraDecodersUtil::MPEG::H265::FSequenceParameterSet, ESPMode::ThreadSafe> SPS;
+		int64 DecoderPTS = -1;
 	};
 
 	enum class EDecodeState
@@ -209,6 +218,7 @@ private:
 	TSharedPtr<IElectraH265VideoDecoderAndroidJava, ESPMode::ThreadSafe> DecoderInstance;
 	IElectraH265VideoDecoderAndroidJava::FDecoderInformation DecoderInfo;
 	EDecodeState DecodeState = EDecodeState::Decoding;
+	TOptional<int64> DecoderPTSOffset;
 	int64 LastPushedPresentationTime = -1;
 	bool bDidSendEOS = false;
 	bool bNewDecoderRequired = false;
@@ -400,6 +410,7 @@ bool FElectraVideoDecoderH265_Android::ResetToCleanStart()
 	CurrentOutput.Reset();
 	InDecoderInput.Empty();
 	DecodeState = EDecodeState::Decoding;
+	DecoderPTSOffset.Reset();
 	LastPushedPresentationTime = -1;
 	bDidSendEOS = false;
 
@@ -460,6 +471,7 @@ IElectraDecoder::EDecoderError FElectraVideoDecoderH265_Android::DecodeAccessUni
 		bTriggerReplay = true;
 		CurrentOutput.Reset();
 		InDecoderInput.Empty();
+		DecoderPTSOffset.Reset();
 		LastPushedPresentationTime = -1;
 		bNewDecoderRequired = true;
 	}
@@ -488,8 +500,20 @@ IElectraDecoder::EDecoderError FElectraVideoDecoderH265_Android::DecodeAccessUni
 		int32 InputBufferIndex = DecoderInstance->DequeueInputBuffer(0);
 		if (InputBufferIndex >= 0)
 		{
-			// Android decoders need the PTS in microseconds
-			int64 PTS = InInputAccessUnit.PTS.GetTicks() / 10;
+			// The decoder could have issues when the PTS is large as is common in a Live stream.
+			// We rebase it to work around this issue.
+			int64 PTSTicks = InInputAccessUnit.PTS.GetTicks();
+			if (!DecoderPTSOffset.IsSet())
+			{
+				DecoderPTSOffset = PlatformElectraDecodersAndroid::CVarElectraDecoderAndroidRebasePTS.GetValueOnAnyThread() ? PTSTicks : 0;
+				DETAILLOG(LogElectraDecoders, Log, TEXT("VideoDecoderH265::DecodeAccessUnit() - Rebase PTS offset %lld"), (long long int)DecoderPTSOffset.GetValue());
+			}
+			PTSTicks -= DecoderPTSOffset.GetValue();
+			// Also the decoder internally truncates nanoseconds to microseconds which it then
+			// sticks onto the decoded output, so we cannot use the PTS as a "user value".
+			// We truncate the value ourselves to account for this.
+			PTSTicks /= 10;		// Convert from HNS to microseconds
+
 			int32 Result = 0;
 			// If this is a sync sample we prepend the CSD. While this is not necessary on a running stream we need to have the CSD
 			// on the first frame and it is easier to prepend it to all IDR frames when seeking etc.
@@ -503,22 +527,24 @@ IElectraDecoder::EDecoderError FElectraVideoDecoderH265_Android::DecodeAccessUni
 				}
 
 				TArray<uint8> CSD = ElectraDecodersUtil::GetVariantValueUInt8Array(InAdditionalOptions, TEXT("csd"));
-				Result = DecoderInstance->QueueInputBuffer(InputBufferIndex, CSD.GetData(), CSD.Num(), InInputAccessUnit.Data, InInputAccessUnit.DataSize, PTS);
+				Result = DecoderInstance->QueueInputBuffer(InputBufferIndex, CSD.GetData(), CSD.Num(), InInputAccessUnit.Data, InInputAccessUnit.DataSize, PTSTicks);
 			}
 			else
 			{
-				Result = DecoderInstance->QueueInputBuffer(InputBufferIndex, nullptr, 0, InInputAccessUnit.Data, InInputAccessUnit.DataSize, PTS);
+				Result = DecoderInstance->QueueInputBuffer(InputBufferIndex, nullptr, 0, InInputAccessUnit.Data, InInputAccessUnit.DataSize, PTSTicks);
 			}
-			DETAILLOG(LogElectraDecoders, Log, TEXT("VideoDecoderH265::DecodeAccessUnit() - queue input %lld, 0x%x"), (long long int)InInputAccessUnit.PTS.GetTicks(), InInputAccessUnit.Flags);
+			DETAILLOG(LogElectraDecoders, Log, TEXT("VideoDecoderH265::DecodeAccessUnit() - queue input %lld (%lld), 0x%x"), (long long int)PTSTicks, (long long int)InInputAccessUnit.PTS.GetTicks(), InInputAccessUnit.Flags);
 			if (Result == 0)
 			{
-				LastPushedPresentationTime = PTS;
+				LastPushedPresentationTime = PTSTicks;
 
 				// Add to the list of inputs passed to the decoder.
 				TSharedPtr<FDecoderInput, ESPMode::NotThreadSafe> In(new FDecoderInput);
 				In->AdditionalOptions = InAdditionalOptions;
 				In->AccessUnit = InInputAccessUnit;
 				In->SPS = CurrentSPS;
+				In->DecoderPTS = PTSTicks;
+
 				// Zero the input pointer and size in the copy. That data is not owned by us and it's best not to have any
 				// values here that would lead us to think that we do.
 				In->AccessUnit.Data = nullptr;
@@ -599,6 +625,7 @@ IElectraDecoder::EDecoderError FElectraVideoDecoderH265_Android::Flush()
 	}
 	CurrentOutput.Reset();
 	InDecoderInput.Empty();
+	DecoderPTSOffset.Reset();
 	LastPushedPresentationTime = -1;
 	bDidSendEOS = false;
 	DecodeState = EDecodeState::Decoding;
@@ -941,8 +968,7 @@ FElectraVideoDecoderH265_Android::EConvertResult FElectraVideoDecoderH265_Androi
 	TSharedPtr<FDecoderInput, ESPMode::NotThreadSafe> In;
 	for(int32 nInDec=0; nInDec<InDecoderInput.Num(); ++nInDec)
 	{
-		// Decoder timestamps are in microseconds, so we need to compare as that.
-		if (InDecoderInput[nInDec]->AccessUnit.PTS / 10 == InInfo.PresentationTimestamp)
+		if (InDecoderInput[nInDec]->DecoderPTS == InInfo.PresentationTimestamp)
 		{
 			In = InDecoderInput[nInDec];
 			InDecoderInput.RemoveAt(nInDec);
@@ -963,14 +989,14 @@ FElectraVideoDecoderH265_Android::EConvertResult FElectraVideoDecoderH265_Androi
 	// If this is a replay frame or a frame that will not be output, we do not need to set it up and merely release the buffer immediately.
 	if ((In->AccessUnit.Flags & (EElectraDecoderFlags::IsReplaySample | EElectraDecoderFlags::DoNotOutput)) != EElectraDecoderFlags::None)
 	{
-		DETAILLOG(LogElectraDecoders, Log, TEXT("VideoDecoderH265::ConvertDecoderOutput() - Discard replay output %lld, 0x%x"), (long long int)In->AccessUnit.PTS.GetTicks(), In->AccessUnit.Flags);
+		DETAILLOG(LogElectraDecoders, Log, TEXT("VideoDecoderH265::ConvertDecoderOutput() - Discard output %lld (%lld), 0x%x"), (long long int)In->DecoderPTS, (long long int)In->AccessUnit.PTS.GetTicks(), In->AccessUnit.Flags);
 		if (DecoderInstance.IsValid())
 		{
 			DecoderInstance->ReleaseOutputBuffer(InInfo.BufferIndex, InInfo.ValidCount, false, -1);
 		}
 		return EConvertResult::Success;
 	}
-	DETAILLOG(LogElectraDecoders, Log, TEXT("VideoDecoderH265::ConvertDecoderOutput() - Convert output %lld, 0x%x"), (long long int)In->AccessUnit.PTS.GetTicks(), In->AccessUnit.Flags);
+	DETAILLOG(LogElectraDecoders, Log, TEXT("VideoDecoderH265::ConvertDecoderOutput() - Convert output %lld (%lld), 0x%x"), (long long int)In->DecoderPTS, (long long int)In->AccessUnit.PTS.GetTicks(), In->AccessUnit.Flags);
 
 	TSharedPtr<FElectraVideoDecoderOutputH265_Android, ESPMode::ThreadSafe> NewOutput = MakeShared<FElectraVideoDecoderOutputH265_Android>();
 	NewOutput->PTS = In->AccessUnit.PTS;
diff --git a/Engine/Plugins/Media/ElectraCodecs/Source/ElectraDecoders/Private/Android/vpx/VideoDecoderVPx_Android.cpp b/Engine/Plugins/Media/ElectraCodecs/Source/ElectraDecoders/Private/Android/vpx/VideoDecoderVPx_Android.cpp
--- a/Engine/Plugins/Media/ElectraCodecs/Source/ElectraDecoders/Private/Android/vpx/VideoDecoderVPx_Android.cpp
+++ b/Engine/Plugins/Media/ElectraCodecs/Source/ElectraDecoders/Private/Android/vpx/VideoDecoderVPx_Android.cpp
@@ -18,6 +18,14 @@
 
 #include "vpx/VideoDecoderVPx_JavaWrapper_Android.h"
 
+/*********************************************************************************************************************/
+
+#include "HAL/IConsoleManager.h"
+namespace PlatformElectraDecodersAndroid
+{
+extern TAutoConsoleVariable<bool> CVarElectraDecoderAndroidRebasePTS;
+}
+
 /*********************************************************************************************************************/
 #define ENABLE_DETAILED_LOG 0
 #if ENABLE_DETAILED_LOG
@@ -174,6 +182,7 @@ private:
 		TMap<FString, FVariant> AdditionalOptions;
 		ElectraDecodersUtil::VPxVideo::FVP8UncompressedHeader HdrVP8;
 		ElectraDecodersUtil::VPxVideo::FVP9UncompressedHeader HdrVP9;
+		int64 DecoderPTS = -1;
 	};
 
 	enum class EDecodeState
@@ -212,6 +221,7 @@ private:
 	TSharedPtr<IElectraVPxVideoDecoderAndroidJava, ESPMode::ThreadSafe> DecoderInstance;
 	IElectraVPxVideoDecoderAndroidJava::FDecoderInformation DecoderInfo;
 	EDecodeState DecodeState = EDecodeState::Decoding;
+	TOptional<int64> DecoderPTSOffset;
 	int64 LastPushedPresentationTime = -1;
 	bool bDidSendEOS = false;
 	bool bNewDecoderRequired = false;
@@ -359,6 +369,7 @@ bool FElectraVideoDecoderVPx_Android::ResetToCleanStart()
 	CurrentOutput.Reset();
 	InDecoderInput.Empty();
 	DecodeState = EDecodeState::Decoding;
+	DecoderPTSOffset.Reset();
 	LastPushedPresentationTime = -1;
 	bDidSendEOS = false;
 
@@ -409,6 +420,7 @@ IElectraDecoder::EDecoderError FElectraVideoDecoderVPx_Android::DecodeAccessUnit
 		bTriggerReplay = true;
 		CurrentOutput.Reset();
 		InDecoderInput.Empty();
+		DecoderPTSOffset.Reset();
 		LastPushedPresentationTime = -1;
 		bNewDecoderRequired = true;
 	}
@@ -456,24 +468,36 @@ IElectraDecoder::EDecoderError FElectraVideoDecoderVPx_Android::DecodeAccessUnit
 		int32 InputBufferIndex = DecoderInstance->DequeueInputBuffer(0);
 		if (InputBufferIndex >= 0)
 		{
-			// Android decoders need the PTS in microseconds
-			int64 PTS = InInputAccessUnit.PTS.GetTicks() / 10;
+			// The decoder could have issues when the PTS is large as is common in a Live stream.
+			// We rebase it to work around this issue.
+			int64 PTSTicks = InInputAccessUnit.PTS.GetTicks();
+			if (!DecoderPTSOffset.IsSet())
+			{
+				DecoderPTSOffset = PlatformElectraDecodersAndroid::CVarElectraDecoderAndroidRebasePTS.GetValueOnAnyThread() ? PTSTicks : 0;
+				DETAILLOG(LogElectraDecoders, Log, TEXT("VideoDecoderVPx::DecodeAccessUnit() - Rebase PTS offset %lld"), (long long int)DecoderPTSOffset.GetValue());
+			}
+			PTSTicks -= DecoderPTSOffset.GetValue();
+			// Also the decoder internally truncates nanoseconds to microseconds which it then
+			// sticks onto the decoded output, so we cannot use the PTS as a "user value".
+			// We truncate the value ourselves to account for this.
+			PTSTicks /= 10;		// Convert from HNS to microseconds
+
 			int32 Result = 0;
 			// If this is a sync sample we prepend the CSD. While this is not necessary on a running stream we need to have the CSD
 			// on the first frame and it is easier to prepend it to all IDR frames when seeking etc.
 			if ((InInputAccessUnit.Flags & EElectraDecoderFlags::IsSyncSample) != EElectraDecoderFlags::None)
 			{
 				TArray<uint8> CSD = ElectraDecodersUtil::GetVariantValueUInt8Array(InAdditionalOptions, TEXT("csd"));
-				Result = DecoderInstance->QueueInputBuffer(InputBufferIndex, CSD.GetData(), CSD.Num(), InInputAccessUnit.Data, InInputAccessUnit.DataSize, PTS);
+				Result = DecoderInstance->QueueInputBuffer(InputBufferIndex, CSD.GetData(), CSD.Num(), InInputAccessUnit.Data, InInputAccessUnit.DataSize, PTSTicks);
 			}
 			else
 			{
-				Result = DecoderInstance->QueueInputBuffer(InputBufferIndex, nullptr, 0, InInputAccessUnit.Data, InInputAccessUnit.DataSize, PTS);
+				Result = DecoderInstance->QueueInputBuffer(InputBufferIndex, nullptr, 0, InInputAccessUnit.Data, InInputAccessUnit.DataSize, PTSTicks);
 			}
-			DETAILLOG(LogElectraDecoders, Log, TEXT("VideoDecoderVPx::DecodeAccessUnit() - queue input %lld, 0x%x"), (long long int)InInputAccessUnit.PTS.GetTicks(), InInputAccessUnit.Flags);
+			DETAILLOG(LogElectraDecoders, Log, TEXT("VideoDecoderVPx::DecodeAccessUnit() - queue input %lld (%lld), 0x%x"), (long long int)PTSTicks, (long long int)InInputAccessUnit.PTS.GetTicks(), InInputAccessUnit.Flags);
 			if (Result == 0)
 			{
-				LastPushedPresentationTime = PTS;
+				LastPushedPresentationTime = PTSTicks;
 
 				// Add to the list of inputs passed to the decoder.
 				TSharedPtr<FDecoderInput, ESPMode::NotThreadSafe> In(new FDecoderInput);
@@ -481,6 +505,7 @@ IElectraDecoder::EDecoderError FElectraVideoDecoderVPx_Android::DecodeAccessUnit
 				In->AccessUnit = InInputAccessUnit;
 				In->HdrVP8 = HdrVP8;
 				In->HdrVP9 = HdrVP9;
+				In->DecoderPTS = PTSTicks;
 				// Zero the input pointer and size in the copy. That data is not owned by us and it's best not to have any
 				// values here that would lead us to think that we do.
 				In->AccessUnit.Data = nullptr;
@@ -561,6 +586,7 @@ IElectraDecoder::EDecoderError FElectraVideoDecoderVPx_Android::Flush()
 	}
 	CurrentOutput.Reset();
 	InDecoderInput.Empty();
+	DecoderPTSOffset.Reset();
 	LastPushedPresentationTime = -1;
 	bDidSendEOS = false;
 	DecodeState = EDecodeState::Decoding;
@@ -881,8 +907,7 @@ FElectraVideoDecoderVPx_Android::EConvertResult FElectraVideoDecoderVPx_Android:
 	TSharedPtr<FDecoderInput, ESPMode::NotThreadSafe> In;
 	for(int32 nInDec=0; nInDec<InDecoderInput.Num(); ++nInDec)
 	{
-		// Decoder timestamps are in microseconds, so we need to compare as that.
-		if (InDecoderInput[nInDec]->AccessUnit.PTS / 10 == InInfo.PresentationTimestamp)
+		if (InDecoderInput[nInDec]->DecoderPTS == InInfo.PresentationTimestamp)
 		{
 			In = InDecoderInput[nInDec];
 			InDecoderInput.RemoveAt(nInDec);
@@ -898,14 +923,14 @@ FElectraVideoDecoderVPx_Android::EConvertResult FElectraVideoDecoderVPx_Android:
 	// If this is a replay frame or a frame that will not be output, we do not need to set it up and merely release the buffer immediately.
 	if ((In->AccessUnit.Flags & (EElectraDecoderFlags::IsReplaySample | EElectraDecoderFlags::DoNotOutput)) != EElectraDecoderFlags::None)
 	{
-		DETAILLOG(LogElectraDecoders, Log, TEXT("VideoDecoderVPx::ConvertDecoderOutput() - Discard replay output %lld, 0x%x"), (long long int)In->AccessUnit.PTS.GetTicks(), In->AccessUnit.Flags);
+		DETAILLOG(LogElectraDecoders, Log, TEXT("VideoDecoderVPx::ConvertDecoderOutput() - Discard output %lld (%lld), 0x%x"), (long long int)In->DecoderPTS, (long long int)In->AccessUnit.PTS.GetTicks(), In->AccessUnit.Flags);
 		if (DecoderInstance.IsValid())
 		{
 			DecoderInstance->ReleaseOutputBuffer(InInfo.BufferIndex, InInfo.ValidCount, false, -1);
 		}
 		return EConvertResult::Success;
 	}
-	DETAILLOG(LogElectraDecoders, Log, TEXT("VideoDecoderVPx::ConvertDecoderOutput() - Convert output %lld, 0x%x"), (long long int)In->AccessUnit.PTS.GetTicks(), In->AccessUnit.Flags);
+	DETAILLOG(LogElectraDecoders, Log, TEXT("VideoDecoderVPx::ConvertDecoderOutput() - Convert output %lld (%lld), 0x%x"), (long long int)In->DecoderPTS, (long long int)In->AccessUnit.PTS.GetTicks(), In->AccessUnit.Flags);
 
 	IElectraVPxVideoDecoderAndroidJava::FOutputFormatInfo Info;
 	if (!DecoderInstance.IsValid() || DecoderInstance->GetOutputFormatInfo(Info, InInfo.BufferIndex) != 0)
