[ElectraDecoders] Expose MFSample on DX H264 decoder output for CPU buffer access

#Synced-CL 51630826

diff --git a/Engine/Plugins/Media/ElectraCodecs/Source/ElectraDecoders/Private/DX/VideoDecoderH264_DX.cpp b/Engine/Plugins/Media/ElectraCodecs/Source/ElectraDecoders/Private/DX/VideoDecoderH264_DX.cpp
--- a/Engine/Plugins/Media/ElectraCodecs/Source/ElectraDecoders/Private/DX/VideoDecoderH264_DX.cpp
+++ b/Engine/Plugins/Media/ElectraCodecs/Source/ElectraDecoders/Private/DX/VideoDecoderH264_DX.cpp
@@ -101,7 +101,9 @@ public:
 #if ELECTRA_HAVE_IMAGEBUFFERS
 			case EElectraDecoderPlatformOutputHandleType::ImageBuffers:
 			{
-				return !MFSample.IsValid() ? static_cast<IElectraDecoderVideoOutputImageBuffers*>(const_cast<FElectraVideoDecoderOutputH264_DX*>(this)) : nullptr;
+				// Check for the GPU resource directly rather than inferring from !MFSample.IsValid(),
+				// since both will coexist when force_cpu_output is set.
+				return Buffer.Resource.IsValid() ? static_cast<IElectraDecoderVideoOutputImageBuffers*>(const_cast<FElectraVideoDecoderOutputH264_DX*>(this)) : nullptr;
 			}
 #endif
 			case EElectraDecoderPlatformOutputHandleType::MFSample:
@@ -367,6 +369,7 @@ private:
 	TRefCountPtr<IMFMediaType> CurrentOutputMediaType;
 
 	TMap<FString, FVariant> InitialCreationOptions;
+	bool bForceCPUOutput = false;
 	uint32 MaxOutputBuffers;
 	uint32 MaxWidth;
 	uint32 MaxHeight;
@@ -396,6 +399,8 @@ FElectraVideoDecoderH264_DX::FElectraVideoDecoderH264_DX(const TMap<FString, FVa
 
 	MaxOutputBuffers = (uint32)ElectraDecodersUtil::GetVariantValueSafeU64(InOptions, TEXT("max_output_buffers"), 5);
 	MaxOutputBuffers += kElectraDecoderPipelineExtraFrames;
+
+	bForceCPUOutput = !!ElectraDecodersUtil::GetVariantValueSafeI64(InOptions, TEXT("force_cpu_output"), 0);
 }
 
 
@@ -1216,6 +1221,11 @@ bool FElectraVideoDecoderH264_DX::ConvertDecoderOutput()
 
 		NewOutput->Pitch = BufferPitch;
 		CopyBufferData(Buffer2D, NewOutput->Buffer.Resource, NewOutput->Buffer.Fence, NewOutput->Buffer.FenceValue, BufferPitch, NewOutput->DecodedHeight);
+
+		if (bForceCPUOutput)
+		{
+			NewOutput->MFSample = DecodedOutputSample;
+		}
 #else
 		//
 		// Post SDK 22621: direct link from MFSample to DX12 resource
@@ -1241,6 +1251,12 @@ bool FElectraVideoDecoderH264_DX::ConvertDecoderOutput()
 			NewOutput->Buffer.Resource = Resource;
 			NewOutput->Buffer.DecoderSync = DecoderSync;
 			NewOutput->Buffer.MFSample = DecodedOutputSample;	// we keep a reference around until we actually submit the sync object to a DX12 queue
+
+			if (bForceCPUOutput)
+			{
+				// Also expose the MFSample via GetPlatformOutputHandle(MFSample) for CPU buffer extraction
+				NewOutput->MFSample = DecodedOutputSample;
+			}
 		}
 #endif // !ALLOW_MFSAMPLE_WITH_DX12
 	}
